Make IncludeTool handle arbitrary Linux toolchains.

- Fixes CIS breakage (UE-43663).

(Edigrating CL 3381266 from Dev-Editor to Main)

diff --git a/Engine/Source/Programs/IncludeTool/IncludeTool/Program.cs b/Engine/Source/Programs/IncludeTool/IncludeTool/Program.cs
--- a/Engine/Source/Programs/IncludeTool/IncludeTool/Program.cs
+++ b/Engine/Source/Programs/IncludeTool/IncludeTool/Program.cs
@@ -346,12 +346,9 @@ namespace IncludeTool
 
 				// Find any additional system include paths for the internal preprocessor, which are normally inserted by the compiler
 				List<DirectoryReference> ExtraSystemIncludePaths = new List<DirectoryReference>();
-				if(Options.Platform == TargetPlatform.Linux && FileToCompileEnvironment.Count > 0)
+				if (Options.Platform == TargetPlatform.Linux && FileToCompileEnvironment.Count > 0)
 				{
-					DirectoryReference SysRootDir = FileToCompileEnvironment.First().Value.Options.Where(x => x.Name.StartsWith("--sysroot=")).Select(x => new DirectoryReference(x.Name.Substring(10))).First();
-					ExtraSystemIncludePaths.Add(DirectoryReference.Combine(SysRootDir, "usr", "include"));
-					ExtraSystemIncludePaths.Add(DirectoryReference.Combine(SysRootDir, "lib", "clang", "3.9.0", "include"));
-					ExtraSystemIncludePaths.Add(DirectoryReference.Combine(InputDir, "Engine", "Source", "ThirdParty", "Linux", "LibCxx", "include", "c++", "v1"));
+					Utility.GetSystemIncludeDirectories(FileToCompileEnvironment.First().Value, InputDir, ExtraSystemIncludePaths);
 				}
 
 				// Preprocess all the input files
diff --git a/Engine/Source/Programs/IncludeTool/IncludeTool/Support/Utility.cs b/Engine/Source/Programs/IncludeTool/IncludeTool/Support/Utility.cs
--- a/Engine/Source/Programs/IncludeTool/IncludeTool/Support/Utility.cs
+++ b/Engine/Source/Programs/IncludeTool/IncludeTool/Support/Utility.cs
@@ -9,6 +9,7 @@ using System.Security.Cryptography;
 using System.Text;
 using System.Threading;
 using System.Threading.Tasks;
+using System.Text.RegularExpressions;
 
 namespace IncludeTool.Support
 {
@@ -282,5 +283,62 @@ namespace IncludeTool.Support
 
 			State.Complete(EndValue - BeginValue, Message, Log);
 		}
+
+		/// <summary>
+		/// Queries compiler for the version and returns string like "3.9.0" or such.
+		/// </summary>
+		private static string DetermineClangVersionDir(string ClangExecutable)
+		{
+			string CompilerVersionString = null;
+
+			Process Proc = new Process();
+			Proc.StartInfo.UseShellExecute = false;
+			Proc.StartInfo.CreateNoWindow = true;
+			Proc.StartInfo.RedirectStandardOutput = true;
+			Proc.StartInfo.RedirectStandardError = true;
+
+			Proc.StartInfo.FileName = ClangExecutable;
+			Proc.StartInfo.Arguments = " --version";
+
+			Proc.Start();
+			Proc.WaitForExit();
+
+			if (Proc.ExitCode == 0)
+			{
+				// read just the first string
+				string VersionString = Proc.StandardOutput.ReadLine();
+
+				Regex VersionPattern = new Regex("version \\d+(\\.\\d+)+");
+				Match VersionMatch = VersionPattern.Match(VersionString);
+
+				// version match will be like "version 3.3", so remove the "version"
+				if (VersionMatch.Value.StartsWith("version "))
+				{
+					CompilerVersionString = VersionMatch.Value.Replace("version ", "");
+				}
+			}
+
+			return CompilerVersionString;
+		}
+
+		/// <summary>
+		/// Gets include directories that the compiler knows about.
+		/// </summary>
+		public static void GetSystemIncludeDirectories(CompileEnvironment CompileEnv, DirectoryReference InputDir, List<DirectoryReference> ExtraSystemIncludePaths)
+		{
+			// FIXME: assumes clang now, and assumes certain hard-coded dirs. Can be improved by querying the compiler instead.
+
+			DirectoryReference SysRootDir = CompileEnv.Options.Where(x => x.Name.StartsWith("--sysroot=")).Select(x => new DirectoryReference(x.Name.Substring(10))).First();
+
+			string CompilerFullPath = CompileEnv.Compiler.FullName;
+
+			ExtraSystemIncludePaths.Add(DirectoryReference.Combine(SysRootDir, "usr", "include"));
+			string CompilerVersionString = DetermineClangVersionDir(CompilerFullPath);
+			if (!String.IsNullOrEmpty(CompilerVersionString))
+			{
+				ExtraSystemIncludePaths.Add(DirectoryReference.Combine(SysRootDir, "lib", "clang", CompilerVersionString, "include"));
+			}
+			ExtraSystemIncludePaths.Add(DirectoryReference.Combine(InputDir, "Engine", "Source", "ThirdParty", "Linux", "LibCxx", "include", "c++", "v1"));
+		}
 	}
 }
