Fix IncludeTool error; ignore function macros on compiler command line.

diff --git a/Engine/Source/Programs/IncludeTool/IncludeTool/Preprocessor.cs b/Engine/Source/Programs/IncludeTool/IncludeTool/Preprocessor.cs
--- a/Engine/Source/Programs/IncludeTool/IncludeTool/Preprocessor.cs
+++ b/Engine/Source/Programs/IncludeTool/IncludeTool/Preprocessor.cs
@@ -331,7 +331,7 @@ namespace IncludeTool
 				{
 					NameToMacro[Tokens[0].Text] = new PreprocessorMacro(Tokens[0].Text, null, Tokens.Skip(2).ToList());
 				}
-				else
+				else if(Tokens[1].Text != "(")
 				{
 					throw new PreprocessorException("Expected '=' token for definition on command-line");
 				}
