Fix static analysis and HTML5 errors due to use of force-included files.

diff --git a/Engine/Source/Programs/IncludeTool/IncludeTool/CompileEnvironment.cs b/Engine/Source/Programs/IncludeTool/IncludeTool/CompileEnvironment.cs
--- a/Engine/Source/Programs/IncludeTool/IncludeTool/CompileEnvironment.cs
+++ b/Engine/Source/Programs/IncludeTool/IncludeTool/CompileEnvironment.cs
@@ -100,6 +100,11 @@ namespace IncludeTool
 		/// </summary>
 		public List<DirectoryReference> IncludePaths = new List<DirectoryReference>();
 
+		/// <summary>
+		/// List of force-included headers
+		/// </summary>
+		public List<FileReference> ForceIncludeFiles = new List<FileReference>();
+
 		/// <summary>
 		/// Default constructor
 		/// </summary>
@@ -120,6 +125,7 @@ namespace IncludeTool
 			Options = new List<CompileOption>(Other.Options);
 			Definitions = new List<string>(Other.Definitions);
 			IncludePaths = new List<DirectoryReference>(Other.IncludePaths);
+			ForceIncludeFiles = new List<FileReference>(Other.ForceIncludeFiles);
 		}
 
 		/// <summary>
@@ -134,11 +140,13 @@ namespace IncludeTool
 			{
 				Arguments.AddRange(Definitions.Select(x => "-D " + x));
 				Arguments.AddRange(IncludePaths.Select(x => "-I \"" + x + "\""));
+				Arguments.AddRange(ForceIncludeFiles.Select(x => "-include \"" + x + "\""));
 			}
 			else
 			{
 				Arguments.AddRange(Definitions.Select(x => "/D " + x));
 				Arguments.AddRange(IncludePaths.Select(x => "/I \"" + x + "\""));
+				Arguments.AddRange(ForceIncludeFiles.Select(x => "/FI\"" + x + "\""));
 			}
 			return String.Join(" ", Arguments);
 		}
@@ -204,6 +212,18 @@ namespace IncludeTool
 					}
 				}
 
+				foreach (FileReference ForceIncludeFile in ForceIncludeFiles)
+				{
+					if(CompilerType == CompilerType.Clang)
+					{
+						Writer.WriteLine("-include {0}", Utility.QuoteArgument(ForceIncludeFile.FullName.Replace('\\', '/')));
+					}
+					else
+					{
+						Writer.WriteLine("/FI{0}", Utility.QuoteArgument(ForceIncludeFile.FullName));
+					}
+				}
+
 				if(CompilerType == CompilerType.Clang)
 				{
 					Writer.WriteLine(Utility.QuoteArgument(SourceFile.FullName.Replace('\\', '/')));
@@ -292,6 +312,14 @@ namespace IncludeTool
 							{
 								Environment.IncludePaths.Add(new DirectoryReference(DirectoryReference.Combine(BaseDir, Tokens[Idx].Substring(2).Replace("//", "/")).FullName.ToLowerInvariant()));
 							}
+							else if(Tokens[Idx].StartsWith("/FI"))
+							{
+								Environment.ForceIncludeFiles.Add(new FileReference(Tokens[Idx].Substring(2)));
+							}
+							else if(Tokens[Idx] == "-include")
+							{
+								Environment.ForceIncludeFiles.Add(new FileReference(Tokens[++Idx]));
+							}
 							else if (Tokens[Idx] == "-Xclang" || Tokens[Idx] == "-target" || Tokens[Idx] == "--target" || Tokens[Idx] == "-x" || Tokens[Idx] == "-o")
 							{
 								Environment.Options.Add(new CompileOption(Tokens[Idx], Tokens[++Idx]));
diff --git a/Engine/Source/Programs/IncludeTool/IncludeTool/Preprocessor.cs b/Engine/Source/Programs/IncludeTool/IncludeTool/Preprocessor.cs
--- a/Engine/Source/Programs/IncludeTool/IncludeTool/Preprocessor.cs
+++ b/Engine/Source/Programs/IncludeTool/IncludeTool/Preprocessor.cs
@@ -286,7 +286,7 @@ namespace IncludeTool
 		/// </summary>
 		/// <param name="Definitions">Definitions passed in on the command-line</param>
 		/// <param name="IncludePaths">Paths to search when resolving include directives</param>
-		public Preprocessor(FileReference PreludeFile, IEnumerable<string> Definitions, IEnumerable<DirectoryReference> IncludePaths)
+		public Preprocessor(FileReference PreludeFile, IEnumerable<string> Definitions, IEnumerable<DirectoryReference> IncludePaths, IEnumerable<FileReference> ForceIncludedFiles)
 		{
 			DateTime Now = DateTime.Now;
 			AddSingleTokenMacro("__DATE__", TokenType.StringLiteral, String.Format("\"{0} {1,2} {2}\"", Now.ToString("MMM"), Now.Day, Now.Year));
@@ -344,6 +344,16 @@ namespace IncludeTool
 					IncludeDirectories.Add(new IncludeDirectory() { Location = IncludePath, WorkspaceDirectory = Workspace.GetDirectory(IncludePath) });
 				}
 			}
+
+			foreach(FileReference ForceIncludedFile in ForceIncludedFiles)
+			{
+				TextBuffer ForceIncludedText = TextBuffer.FromFile(ForceIncludedFile.FullName);
+				PreprocessorMarkup[] ForceIncludedMarkup = PreprocessorMarkup.ParseArray(ForceIncludedText);
+				foreach(PreprocessorMarkup Markup in ForceIncludedMarkup)
+				{
+					ParseMarkup(Markup.Type, Markup.Tokens, Markup.Location.LineIdx);
+				}
+			}
 		}
 
 		/// <summary>
@@ -1950,7 +1960,7 @@ namespace IncludeTool
         {
 			SourceFile File = new SourceFile(null, TextBuffer.FromString(Fragment), SourceFileFlags.Inline);
 
-			Preprocessor Instance = new Preprocessor(null, new string[] { }, new DirectoryReference[] { });
+			Preprocessor Instance = new Preprocessor(null, new string[] { }, new DirectoryReference[] { }, new FileReference[] { });
 
 			string Result;
 			try
@@ -2056,7 +2066,7 @@ namespace IncludeTool
 			BufferedTextWriter BufferedLog = new BufferedTextWriter();
 
 			// Run the preprocessor on it
-			Preprocessor ToolPreprocessor = new Preprocessor(PreludeFile, Environment.Definitions, Environment.IncludePaths);
+			Preprocessor ToolPreprocessor = new Preprocessor(PreludeFile, Environment.Definitions, Environment.IncludePaths, Environment.ForceIncludeFiles);
 			FileReference ToolOutputFile = FileReference.Combine(ToolDir, InputFile.GetFileNameWithoutExtension() + ".i");
 			ToolPreprocessor.PreprocessFile(InputFile.FullName, ToolOutputFile);
 
diff --git a/Engine/Source/Programs/IncludeTool/IncludeTool/Program.cs b/Engine/Source/Programs/IncludeTool/IncludeTool/Program.cs
--- a/Engine/Source/Programs/IncludeTool/IncludeTool/Program.cs
+++ b/Engine/Source/Programs/IncludeTool/IncludeTool/Program.cs
@@ -933,7 +933,7 @@ namespace IncludeTool
 			IncludePaths.AddRange(ExtraSystemIncludePaths);
 
 			// Set up the preprocessor with the correct defines for this environment
-			Preprocessor PreprocessorInst = new Preprocessor(PreludeLocation, CompileEnvironment.Definitions, IncludePaths);
+			Preprocessor PreprocessorInst = new Preprocessor(PreludeLocation, CompileEnvironment.Definitions, IncludePaths, CompileEnvironment.ForceIncludeFiles);
 			PreprocessIncludedFile(PreprocessorInst, InitialFile, FileToActiveMarkup, Log);
 		}
 
diff --git a/Engine/Source/Programs/UnrealBuildTool/Platform/HTML5/HTML5ToolChain.cs b/Engine/Source/Programs/UnrealBuildTool/Platform/HTML5/HTML5ToolChain.cs
--- a/Engine/Source/Programs/UnrealBuildTool/Platform/HTML5/HTML5ToolChain.cs
+++ b/Engine/Source/Programs/UnrealBuildTool/Platform/HTML5/HTML5ToolChain.cs
@@ -442,6 +442,12 @@ namespace UnrealBuildTool
 				Arguments += string.Format(" -D__EMSCRIPTEN_TRACING__");
 			}
 
+			// Force include all the requested headers
+			foreach(FileReference ForceIncludeFile in CompileEnvironment.ForceIncludeFiles)
+			{
+				Arguments += String.Format(" -include \"{0}\"", ForceIncludeFile.FullName);
+			}
+
 			foreach (FileItem SourceFile in InputFiles)
 			{
 				Action CompileAction = ActionGraph.Add(ActionType.Compile);
diff --git a/Engine/Source/Programs/UnrealBuildTool/Platform/Windows/PVSToolChain.cs b/Engine/Source/Programs/UnrealBuildTool/Platform/Windows/PVSToolChain.cs
--- a/Engine/Source/Programs/UnrealBuildTool/Platform/Windows/PVSToolChain.cs
+++ b/Engine/Source/Programs/UnrealBuildTool/Platform/Windows/PVSToolChain.cs
@@ -47,6 +47,10 @@ namespace UnrealBuildTool
 			{
 				SharedArguments.Add(String.Format("/D \"{0}\"", Definition));
 			}
+			foreach(FileReference ForceIncludeFile in CompileEnvironment.ForceIncludeFiles)
+			{
+				SharedArguments.Add(String.Format("/FI\"{0}\"", ForceIncludeFile.FullName));
+			}
 
 			// Get the path to PVS studio
 			FileReference AnalyzerFile = new FileReference(Path.Combine(Environment.GetFolderPath(Environment.SpecialFolder.ProgramFilesX86), "PVS-Studio", "x64", "PVS-Studio.exe"));
