PR #5052: Remove OutputLines = new List<string>(); (Contributed by [REDACTED_PERSON])

diff --git a/Engine/Source/Programs/IncludeTool/IncludeTool/ManagedProcess.cs b/Engine/Source/Programs/IncludeTool/IncludeTool/ManagedProcess.cs
--- a/Engine/Source/Programs/IncludeTool/IncludeTool/ManagedProcess.cs
+++ b/Engine/Source/Programs/IncludeTool/IncludeTool/ManagedProcess.cs
@@ -1,4 +1,4 @@
-﻿// Copyright 1998-2018 Epic Games, Inc. All Rights Reserved.
+// Copyright 1998-2018 Epic Games, Inc. All Rights Reserved.
 
 using Microsoft.Win32.SafeHandles;
 using System;
@@ -417,7 +417,6 @@ namespace IncludeTool
 
 			// Manually read all the output lines from the stream. Using ReadToEndAsync() or ReadAsync() on the StreamReader is abysmally slow, especially
 			// for high-volume processes. Manually picking out the lines via a buffered ReadAsync() call was found to be 6x faster on 'p4 -ztag have' calls.
-			OutputLines = new List<string>();
 			byte[] Buffer = new byte[32 * 1024];
 			byte LastCharacter = 0;
 			int NumBytesInBuffer = 0;
