CIS fix for IncludeTool.

#rnx

diff --git a/Engine/Source/Programs/IncludeTool/IncludeTool/Program.cs b/Engine/Source/Programs/IncludeTool/IncludeTool/Program.cs
--- a/Engine/Source/Programs/IncludeTool/IncludeTool/Program.cs
+++ b/Engine/Source/Programs/IncludeTool/IncludeTool/Program.cs
@@ -1,4 +1,4 @@
-﻿// Copyright 1998-2019 Epic Games, Inc. All Rights Reserved.
+// Copyright 1998-2019 Epic Games, Inc. All Rights Reserved.
 
 using IncludeTool.Reports;
 using IncludeTool.Support;
@@ -862,7 +862,11 @@ namespace IncludeTool
 			DirectoryReference WorkingDir = DirectoryReference.Combine(RootDir, "Engine");
 
 			FileReference UnrealBuildTool = FileReference.Combine(RootDir, "Engine", "Binaries", "DotNET", "UnrealBuildTool.exe");
-			if (Utility.Run(UnrealBuildTool, String.Format("{0} {1} {2}{3} -disableunity -xgeexport -nobuilduht -nopch -nodebuginfo -jsonexport=\"{4}\"", Target, Configuration, Platform, Precompile? " -precompile" : "", TaskListFile.ChangeExtension(".json").FullName), WorkingDir, Log) != 0)
+			if (Utility.Run(UnrealBuildTool, String.Format("-Mode=JsonExport {0} {1} {2}{3} -disableunity -xgeexport -nobuilduht -nopch -nodebuginfo -outputfile=\"{4}\"", Target, Configuration, Platform, Precompile? " -precompile" : "", TaskListFile.ChangeExtension(".json").FullName), WorkingDir, Log) != 0)
+			{
+				throw new Exception("UnrealBuildTool failed");
+			}
+			if (Utility.Run(UnrealBuildTool, String.Format("{0} {1} {2}{3} -disableunity -xgeexport -nobuilduht -nopch -nodebuginfo", Target, Configuration, Platform, Precompile? " -precompile" : ""), WorkingDir, Log) != 0)
 			{
 				throw new Exception("UnrealBuildTool failed");
 			}
