Fix IncludeTool warnings and errors.

#rnx

diff --git a/Engine/Plugins/Experimental/ControlRig/Source/ControlRigEditor/Private/Editor/ControlRigEditor.h b/Engine/Plugins/Experimental/ControlRig/Source/ControlRigEditor/Private/Editor/ControlRigEditor.h
--- a/Engine/Plugins/Experimental/ControlRig/Source/ControlRigEditor/Private/Editor/ControlRigEditor.h
+++ b/Engine/Plugins/Experimental/ControlRig/Source/ControlRigEditor/Private/Editor/ControlRigEditor.h
@@ -5,7 +5,7 @@
 #include "IControlRigEditor.h"
 #include "ControlRigEditorEditMode.h"
 #include "AssetEditorModeManager.h"
-#include "DragAndDrop//GraphNodeDragDropOp.h"
+#include "DragAndDrop/GraphNodeDragDropOp.h"
 #include "ControlRigDefines.h"
 #include "ControlRigLog.h"
 #include "Drawing/ControlRigDrawInterface.h"
diff --git a/Engine/Plugins/Experimental/ControlRig/Source/ControlRigEditor/Private/Graph/SControlRigGraphNodeComment.cpp b/Engine/Plugins/Experimental/ControlRig/Source/ControlRigEditor/Private/Graph/SControlRigGraphNodeComment.cpp
--- a/Engine/Plugins/Experimental/ControlRig/Source/ControlRigEditor/Private/Graph/SControlRigGraphNodeComment.cpp
+++ b/Engine/Plugins/Experimental/ControlRig/Source/ControlRigEditor/Private/Graph/SControlRigGraphNodeComment.cpp
@@ -6,7 +6,7 @@
 #include "PropertyPathHelpers.h"
 #include "UObject/PropertyPortFlags.h"
 #include "ControlRigBlueprint.h"
-#include "Graph//ControlRigGraph.h"
+#include "Graph/ControlRigGraph.h"
 #include "ControlRigController.h"
 #if WITH_EDITOR
 #include "Editor.h"
diff --git a/Engine/Plugins/Experimental/RemoteSession/Source/RemoteSession/RemoteSession.h b/Engine/Plugins/Experimental/RemoteSession/Source/RemoteSession/RemoteSession.h
--- a/Engine/Plugins/Experimental/RemoteSession/Source/RemoteSession/RemoteSession.h
+++ b/Engine/Plugins/Experimental/RemoteSession/Source/RemoteSession/RemoteSession.h
@@ -5,7 +5,6 @@
 #include "CoreMinimal.h"
 #include "Modules/ModuleManager.h"
 #include "RemoteSessionRole.h"
-#include "Channels/RemoteSessionChannel.h"
 
 #define REMOTE_SESSION_VERSION_STRING TEXT("1.0.5")
 
diff --git a/Engine/Plugins/Media/MediaPlayerEditor/Source/MediaPlayerEditor/Public/Widgets/SMediaImage.h b/Engine/Plugins/Media/MediaPlayerEditor/Source/MediaPlayerEditor/Public/Widgets/SMediaImage.h
--- a/Engine/Plugins/Media/MediaPlayerEditor/Source/MediaPlayerEditor/Public/Widgets/SMediaImage.h
+++ b/Engine/Plugins/Media/MediaPlayerEditor/Source/MediaPlayerEditor/Public/Widgets/SMediaImage.h
@@ -4,7 +4,7 @@
 
 #include "UObject/GCObject.h"
 
-#include "Widgets//DeclarativeSyntaxSupport.h"
+#include "Widgets/DeclarativeSyntaxSupport.h"
 #include "Widgets/SCompoundWidget.h"
 
 class UMaterial;
diff --git a/Engine/Source/Programs/IncludeTool/IncludeTool/Preprocessor.cs b/Engine/Source/Programs/IncludeTool/IncludeTool/Preprocessor.cs
--- a/Engine/Source/Programs/IncludeTool/IncludeTool/Preprocessor.cs
+++ b/Engine/Source/Programs/IncludeTool/IncludeTool/Preprocessor.cs
@@ -1,4 +1,4 @@
-﻿// Copyright 1998-2019 Epic Games, Inc. All Rights Reserved.
+// Copyright 1998-2019 Epic Games, Inc. All Rights Reserved.
 
 using System;
 using System.Collections.Generic;
@@ -286,7 +286,7 @@ namespace IncludeTool
 		/// </summary>
 		/// <param name="Definitions">Definitions passed in on the command-line</param>
 		/// <param name="IncludePaths">Paths to search when resolving include directives</param>
-		public Preprocessor(FileReference PreludeFile, IEnumerable<string> Definitions, IEnumerable<DirectoryReference> IncludePaths, IEnumerable<FileReference> ForceIncludedFiles)
+		public Preprocessor(FileReference PreludeFile, IEnumerable<string> Definitions, IEnumerable<DirectoryReference> IncludePaths)
 		{
 			DateTime Now = DateTime.Now;
 			AddSingleTokenMacro("__DATE__", TokenType.StringLiteral, String.Format("\"{0} {1,2} {2}\"", Now.ToString("MMM"), Now.Day, Now.Year));
@@ -344,16 +344,6 @@ namespace IncludeTool
 					IncludeDirectories.Add(new IncludeDirectory() { Location = IncludePath, WorkspaceDirectory = Workspace.GetDirectory(IncludePath) });
 				}
 			}
-
-			foreach(FileReference ForceIncludedFile in ForceIncludedFiles)
-			{
-				TextBuffer ForceIncludedText = TextBuffer.FromFile(ForceIncludedFile.FullName);
-				PreprocessorMarkup[] ForceIncludedMarkup = PreprocessorMarkup.ParseArray(ForceIncludedText);
-				foreach(PreprocessorMarkup Markup in ForceIncludedMarkup)
-				{
-					ParseMarkup(Markup.Type, Markup.Tokens, Markup.Location.LineIdx);
-				}
-			}
 		}
 
 		/// <summary>
@@ -1463,12 +1453,19 @@ namespace IncludeTool
 		/// </summary>
 		/// <param name="InputFile">The file to read from</param>
 		/// <param name="OutputFile">File to output to</param>
-		public void PreprocessFile(string InputFileName, FileReference OutputFile)
+		public void PreprocessFile(string InputFileName, FileReference OutputFile, IEnumerable<FileReference> ForceIncludeFiles)
 		{
 			using (StreamWriter Writer = new StreamWriter(OutputFile.FullName))
 			{
-				WorkspaceFile InputWorkspaceFile = Workspace.GetFile(new FileReference(InputFileName));
-				PreprocessFile(new PreprocessorFile(InputFileName, InputWorkspaceFile), Writer);
+				List<FileReference> InputFiles = new List<FileReference>();
+				InputFiles.AddRange(ForceIncludeFiles);
+				InputFiles.Add(new FileReference(InputFileName));
+
+				foreach(FileReference InputFile in InputFiles)
+				{
+					WorkspaceFile InputWorkspaceFile = Workspace.GetFile(InputFile);
+					PreprocessFile(new PreprocessorFile(InputFile.FullName, InputWorkspaceFile), Writer);
+				}
 			}
 		}
 
@@ -1960,7 +1957,7 @@ namespace IncludeTool
         {
 			SourceFile File = new SourceFile(null, TextBuffer.FromString(Fragment), SourceFileFlags.Inline);
 
-			Preprocessor Instance = new Preprocessor(null, new string[] { }, new DirectoryReference[] { }, new FileReference[] { });
+			Preprocessor Instance = new Preprocessor(null, new string[] { }, new DirectoryReference[] { });
 
 			string Result;
 			try
@@ -2066,9 +2063,9 @@ namespace IncludeTool
 			BufferedTextWriter BufferedLog = new BufferedTextWriter();
 
 			// Run the preprocessor on it
-			Preprocessor ToolPreprocessor = new Preprocessor(PreludeFile, Environment.Definitions, Environment.IncludePaths, Environment.ForceIncludeFiles);
+			Preprocessor ToolPreprocessor = new Preprocessor(PreludeFile, Environment.Definitions, Environment.IncludePaths);
 			FileReference ToolOutputFile = FileReference.Combine(ToolDir, InputFile.GetFileNameWithoutExtension() + ".i");
-			ToolPreprocessor.PreprocessFile(InputFile.FullName, ToolOutputFile);
+			ToolPreprocessor.PreprocessFile(InputFile.FullName, ToolOutputFile, Environment.ForceIncludeFiles);
 
 			// Generate the compiler output
 			FileReference CompilerOutputFile = FileReference.Combine(CompilerDir, InputFile.GetFileNameWithoutExtension() + ".i");
diff --git a/Engine/Source/Programs/IncludeTool/IncludeTool/PreprocessorExpression.cs b/Engine/Source/Programs/IncludeTool/IncludeTool/PreprocessorExpression.cs
--- a/Engine/Source/Programs/IncludeTool/IncludeTool/PreprocessorExpression.cs
+++ b/Engine/Source/Programs/IncludeTool/IncludeTool/PreprocessorExpression.cs
@@ -1,4 +1,4 @@
-﻿// Copyright 1998-2019 Epic Games, Inc. All Rights Reserved.
+// Copyright 1998-2019 Epic Games, Inc. All Rights Reserved.
 
 using IncludeTool.Support;
 using System;
@@ -445,6 +445,21 @@ namespace IncludeTool
 					}
 				}
 			}
+			else if(Token.Length >= 2 && Token.StartsWith("'") && Token.EndsWith("'"))
+			{
+				if(Token.Length == 3)
+				{
+					Value = Token[1];
+				}
+				else if(Token.Length == 6 && Token[1] == '\\' && (Token[2] >= '0' && Token[2] <= '7') && (Token[3] >= '0' && Token[3] <= '7') && (Token[4] >= '0' && Token[4] <= '7'))
+				{
+					Value = ((Token[2] - '0') << 6) | ((Token[1] - '0') << 3) | (Token[2] - '0');
+				}
+				else
+				{
+					throw new PreprocessorException("Invalid character value");
+				}
+			}
 			else
 			{
 				for (int Idx = 0; Idx < Token.Length; Idx++)
diff --git a/Engine/Source/Programs/IncludeTool/IncludeTool/Program.cs b/Engine/Source/Programs/IncludeTool/IncludeTool/Program.cs
--- a/Engine/Source/Programs/IncludeTool/IncludeTool/Program.cs
+++ b/Engine/Source/Programs/IncludeTool/IncludeTool/Program.cs
@@ -937,7 +937,11 @@ namespace IncludeTool
 			IncludePaths.AddRange(ExtraSystemIncludePaths);
 
 			// Set up the preprocessor with the correct defines for this environment
-			Preprocessor PreprocessorInst = new Preprocessor(PreludeLocation, CompileEnvironment.Definitions, IncludePaths, CompileEnvironment.ForceIncludeFiles);
+			Preprocessor PreprocessorInst = new Preprocessor(PreludeLocation, CompileEnvironment.Definitions, IncludePaths);
+			foreach(FileReference ForceIncludedFile in CompileEnvironment.ForceIncludeFiles)
+			{
+				PreprocessIncludedFile(PreprocessorInst, new PreprocessorFile(ForceIncludedFile.FullName, Workspace.GetFile(ForceIncludedFile)), FileToActiveMarkup, Log);
+			}
 			PreprocessIncludedFile(PreprocessorInst, InitialFile, FileToActiveMarkup, Log);
 		}
 
@@ -1002,7 +1006,7 @@ namespace IncludeTool
 				{
 					foreach(PreprocessorMarkup Markup in File.Markup)
 					{
-						if(Markup.IncludedFile != null && (Markup.IncludedFile.Flags & SourceFileFlags.External) == 0 && !File.Location.FullName.Contains("PhyaLib") && !File.Location.GetFileName().Contains("Recast") && !File.Location.GetFileName().Contains("FramePro") && !File.Location.GetFileName().Contains("libSampleRate"))
+						if(Markup.IncludedFile != null && (Markup.IncludedFile.Flags & SourceFileFlags.External) == 0 && !File.Location.FullName.Contains("PhyaLib") && !File.Location.GetFileName().Contains("Recast") && !File.Location.GetFileName().Contains("FramePro") && !File.Location.FullName.Contains("libSampleRate"))
 						{
 							Log.WriteLine("{0}({1}): warning: external file is including non-external file ('{2}')", IncludedFile.Location.FullName, Markup.Location.LineIdx + 1, Markup.IncludedFile.Location);
 						}
diff --git a/Engine/Source/Programs/IncludeTool/IncludeTool/Rules.cs b/Engine/Source/Programs/IncludeTool/IncludeTool/Rules.cs
--- a/Engine/Source/Programs/IncludeTool/IncludeTool/Rules.cs
+++ b/Engine/Source/Programs/IncludeTool/IncludeTool/Rules.cs
@@ -262,6 +262,7 @@ namespace IncludeTool
 			"/Engine/Source/Runtime/Core/Public/UObject/PendingVersions.h",
 			"/Engine/Source/Runtime/CoreUObject/Public/UObject/ScriptSerialization.h",
 			"/Engine/Source/Runtime/Online/HTTP/Public/HttpPackage.h",
+			"/Engine/Plugins/NotForLicensees/OnlineGameplayFramework/Source/McpProfileSys/Public/McpProfileSysPackage.h",
 			"/Engine/Plugins/Online/OnlineSubsystem/Source/Public/OnlineSubsystemPackage.h",
 			"/Engine/Plugins/Online/OnlineSubsystemNull/Source/Public/OnlineSubsystemNullPackage.h",
 			"/Engine/Plugins/Online/NotForLicensees/OnlineSubsystemMcp/Source/Public/OnlineSubsystemMcpPackage.h",
@@ -562,6 +563,10 @@ namespace IncludeTool
 			{
 				return true;
 			}
+			if(Markup.Type == PreprocessorMarkupType.Define && Markup.Tokens[0].Text == "ONLINE_LOG_PREFIX")
+			{
+				return true;
+			}
 			if(Markup.Type == PreprocessorMarkupType.Define && Markup.Tokens[0].Text == "DEPRECATED_FORGAME")
 			{
 				return true;
diff --git a/Engine/Source/Runtime/AudioMixer/Private/AudioMixerSource.cpp b/Engine/Source/Runtime/AudioMixer/Private/AudioMixerSource.cpp
--- a/Engine/Source/Runtime/AudioMixer/Private/AudioMixerSource.cpp
+++ b/Engine/Source/Runtime/AudioMixer/Private/AudioMixerSource.cpp
@@ -1,6 +1,7 @@
 // Copyright 1998-2019 Epic Games, Inc. All Rights Reserved.
 
 #include "AudioMixerSource.h"
+#include "AudioMixerSourceBuffer.h"
 #include "AudioMixerDevice.h"
 #include "AudioMixerSourceVoice.h"
 #include "ActiveSound.h"
diff --git a/Engine/Source/Runtime/AudioMixer/Private/AudioMixerSourceManager.cpp b/Engine/Source/Runtime/AudioMixer/Private/AudioMixerSourceManager.cpp
--- a/Engine/Source/Runtime/AudioMixer/Private/AudioMixerSourceManager.cpp
+++ b/Engine/Source/Runtime/AudioMixer/Private/AudioMixerSourceManager.cpp
@@ -1,6 +1,7 @@
 // Copyright 1998-2019 Epic Games, Inc. All Rights Reserved.
 
 #include "AudioMixerSourceManager.h"
+#include "AudioMixerSourceBuffer.h"
 #include "AudioMixerSource.h"
 #include "AudioMixerDevice.h"
 #include "AudioMixerSourceVoice.h"
diff --git a/Engine/Source/Runtime/AudioMixer/Private/AudioMixerSourceManager.h b/Engine/Source/Runtime/AudioMixer/Private/AudioMixerSourceManager.h
--- a/Engine/Source/Runtime/AudioMixer/Private/AudioMixerSourceManager.h
+++ b/Engine/Source/Runtime/AudioMixer/Private/AudioMixerSourceManager.h
@@ -17,7 +17,6 @@
 #include "DSP/BufferVectorOperations.h"
 #include "IAudioExtensionPlugin.h"
 #include "Containers/Queue.h"
-#include "AudioMixerSourceBuffer.h"
 
 
 namespace Audio
diff --git a/Engine/Source/Runtime/Core/Public/Misc/Compression.h b/Engine/Source/Runtime/Core/Public/Misc/Compression.h
--- a/Engine/Source/Runtime/Core/Public/Misc/Compression.h
+++ b/Engine/Source/Runtime/Core/Public/Misc/Compression.h
@@ -2,7 +2,8 @@
 
 #pragma once
 
-#include "CoreMinimal.h"
+#include "CoreTypes.h"
+#include "UObject/NameTypes.h"
 #include "Templates/Atomic.h"
 #include "Misc/CompressionFlags.h"
 #include "HAL/CriticalSection.h"
diff --git a/Engine/Source/Runtime/Experimental/FieldSystem/Source/FieldSystemCore/Public/Field/FieldSystem.h b/Engine/Source/Runtime/Experimental/FieldSystem/Source/FieldSystemCore/Public/Field/FieldSystem.h
--- a/Engine/Source/Runtime/Experimental/FieldSystem/Source/FieldSystemCore/Public/Field/FieldSystem.h
+++ b/Engine/Source/Runtime/Experimental/FieldSystem/Source/FieldSystemCore/Public/Field/FieldSystem.h
@@ -9,8 +9,6 @@
 #include "Field/FieldSystemTypes.h"
 #include "Math/Vector.h"
 
-#include "Field/FieldSystem.h"
-
 /**
 * MetaData
 *
diff --git a/Engine/Source/Runtime/Experimental/GeometryCollectionEngine/Private/GeometryCollection/GeometryCollectionSceneProxy.h b/Engine/Source/Runtime/Experimental/GeometryCollectionEngine/Private/GeometryCollection/GeometryCollectionSceneProxy.h
--- a/Engine/Source/Runtime/Experimental/GeometryCollectionEngine/Private/GeometryCollection/GeometryCollectionSceneProxy.h
+++ b/Engine/Source/Runtime/Experimental/GeometryCollectionEngine/Private/GeometryCollection/GeometryCollectionSceneProxy.h
@@ -8,14 +8,10 @@
 #include "StaticMeshResources.h"
 #include "Rendering/SkinWeightVertexBuffer.h"
 #include "GeometryCollectionRendering.h"
+#include "GeometryCollection/GeometryCollectionEditorSelection.h"
 #include "HitProxies.h"
 #include "EngineUtils.h"
 
-#ifndef GEOMETRYCOLLECTION_EDITOR_SELECTION
-#define GEOMETRYCOLLECTION_EDITOR_SELECTION WITH_EDITOR
-#endif
-
-
 #if GEOMETRYCOLLECTION_EDITOR_SELECTION
 #include "GeometryCollection/GeometryCollectionHitProxy.h"
 #endif
diff --git a/Engine/Source/Runtime/Experimental/GeometryCollectionEngine/Public/GeometryCollection/GeometryCollectionComponent.h b/Engine/Source/Runtime/Experimental/GeometryCollectionEngine/Public/GeometryCollection/GeometryCollectionComponent.h
--- a/Engine/Source/Runtime/Experimental/GeometryCollectionEngine/Public/GeometryCollection/GeometryCollectionComponent.h
+++ b/Engine/Source/Runtime/Experimental/GeometryCollectionEngine/Public/GeometryCollection/GeometryCollectionComponent.h
@@ -17,6 +17,7 @@
 #include "Physics/Experimental/PhysScene_Chaos.h"
 #include "GeometryCollection/GeometryDynamicCollection.h"
 #include "GeometryCollectionObject.h"
+#include "GeometryCollectionEditorSelection.h"
 #include "GeometryCollection/RecordedTransformTrack.h"
 #include "Templates/UniquePtr.h"
 #include "PhysicalMaterials/Experimental/ChaosPhysicalMaterial.h"
@@ -26,10 +27,6 @@
 
 #include "GeometryCollectionComponent.generated.h"
 
-#ifndef GEOMETRYCOLLECTION_EDITOR_SELECTION
-#define GEOMETRYCOLLECTION_EDITOR_SELECTION WITH_EDITOR
-#endif
-
 struct FGeometryCollectionConstantData;
 struct FGeometryCollectionDynamicData;
 class FGeometryCollectionPhysicsProxy;
diff --git a/Engine/Source/Runtime/Experimental/GeometryCollectionEngine/Public/GeometryCollection/GeometryCollectionEditorSelection.h b/Engine/Source/Runtime/Experimental/GeometryCollectionEngine/Public/GeometryCollection/GeometryCollectionEditorSelection.h
new file mode 100644
--- /dev/null
+++ b/Engine/Source/Runtime/Experimental/GeometryCollectionEngine/Public/GeometryCollection/GeometryCollectionEditorSelection.h
@@ -0,0 +1,7 @@
+// Copyright 1998-2019 Epic Games, Inc. All Rights Reserved.
+
+#pragma once
+
+#ifndef GEOMETRYCOLLECTION_EDITOR_SELECTION
+#define GEOMETRYCOLLECTION_EDITOR_SELECTION WITH_EDITOR
+#endif
diff --git a/Engine/Source/Runtime/Experimental/GeometryCollectionEngine/Public/GeometryCollection/GeometryCollectionHitProxy.h b/Engine/Source/Runtime/Experimental/GeometryCollectionEngine/Public/GeometryCollection/GeometryCollectionHitProxy.h
--- a/Engine/Source/Runtime/Experimental/GeometryCollectionEngine/Public/GeometryCollection/GeometryCollectionHitProxy.h
+++ b/Engine/Source/Runtime/Experimental/GeometryCollectionEngine/Public/GeometryCollection/GeometryCollectionHitProxy.h
@@ -1,9 +1,7 @@
 // Copyright 1998-2019 Epic Games, Inc. All Rights Reserved.
 #pragma once
 
-#ifndef GEOMETRYCOLLECTION_EDITOR_SELECTION
-#define GEOMETRYCOLLECTION_EDITOR_SELECTION WITH_EDITOR
-#endif
+#include "GeometryCollectionEditorSelection.h"
 
 #if GEOMETRYCOLLECTION_EDITOR_SELECTION
 #include "EngineUtils.h"
diff --git a/Engine/Source/Runtime/Online/BackgroundHTTP/Private/BackgroundHttpManagerImpl.cpp b/Engine/Source/Runtime/Online/BackgroundHTTP/Private/BackgroundHttpManagerImpl.cpp
--- a/Engine/Source/Runtime/Online/BackgroundHTTP/Private/BackgroundHttpManagerImpl.cpp
+++ b/Engine/Source/Runtime/Online/BackgroundHTTP/Private/BackgroundHttpManagerImpl.cpp
@@ -1,5 +1,6 @@
 // Copyright 1998-2019 Epic Games, Inc. All Rights Reserved.
 #include "BackgroundHttpManagerImpl.h"
+#include "PlatformBackgroundHttp.h"
 
 #include "HAL/FileManager.h"
 #include "HAL/PlatformFilemanager.h"
diff --git a/Engine/Source/Runtime/Online/BackgroundHTTP/Private/BackgroundHttpModule.cpp b/Engine/Source/Runtime/Online/BackgroundHTTP/Private/BackgroundHttpModule.cpp
--- a/Engine/Source/Runtime/Online/BackgroundHTTP/Private/BackgroundHttpModule.cpp
+++ b/Engine/Source/Runtime/Online/BackgroundHTTP/Private/BackgroundHttpModule.cpp
@@ -2,6 +2,7 @@
 
 #include "BackgroundHttpModule.h"
 #include "BackgroundHttp.h"
+#include "PlatformBackgroundHttp.h"
 #include "Modules/ModuleManager.h"
 #include "Interfaces/IHttpRequest.h"
 
diff --git a/Engine/Source/Runtime/Online/BackgroundHTTP/Public/Interfaces/IBackgroundHttpManager.h b/Engine/Source/Runtime/Online/BackgroundHTTP/Public/Interfaces/IBackgroundHttpManager.h
--- a/Engine/Source/Runtime/Online/BackgroundHTTP/Public/Interfaces/IBackgroundHttpManager.h
+++ b/Engine/Source/Runtime/Online/BackgroundHTTP/Public/Interfaces/IBackgroundHttpManager.h
@@ -2,7 +2,7 @@
 
 #pragma once
 
-#include "BackgroundHttp.h"
+#include "IBackgroundHttpRequest.h"
 
 typedef TSharedPtr<class IBackgroundHttpManager, ESPMode::ThreadSafe> FBackgroundHttpManagerPtr;
 
