More IncludeTool fixes.

#rnx

diff --git a/Engine/Plugins/Animation/LiveLink/Source/LiveLinkEditor/Private/LiveLinkTest.h b/Engine/Plugins/Animation/LiveLink/Source/LiveLinkEditor/Private/LiveLinkTest.h
--- a/Engine/Plugins/Animation/LiveLink/Source/LiveLinkEditor/Private/LiveLinkTest.h
+++ b/Engine/Plugins/Animation/LiveLink/Source/LiveLinkEditor/Private/LiveLinkTest.h
@@ -1,5 +1,7 @@
 // Copyright 1998-2019 Epic Games, Inc. All Rights Reserved.
 
+#pragma once
+
 #include "CoreMinimal.h"
 #include "LiveLinkTypes.h"
 
diff --git a/Engine/Plugins/Experimental/GeometryProcessing/Source/DynamicMesh/Public/DynamicMeshAABBTree3.h b/Engine/Plugins/Experimental/GeometryProcessing/Source/DynamicMesh/Public/DynamicMeshAABBTree3.h
--- a/Engine/Plugins/Experimental/GeometryProcessing/Source/DynamicMesh/Public/DynamicMeshAABBTree3.h
+++ b/Engine/Plugins/Experimental/GeometryProcessing/Source/DynamicMesh/Public/DynamicMeshAABBTree3.h
@@ -1,5 +1,7 @@
 // Copyright 1998-2019 Epic Games, Inc. All Rights Reserved.
 
+#pragma once
+
 #include "Spatial/MeshAABBTree3.h"
 #include "DynamicMesh3.h"
 
diff --git a/Engine/Plugins/Experimental/GeometryProcessing/Source/GeometricObjects/Public/Spatial/FastWinding.h b/Engine/Plugins/Experimental/GeometryProcessing/Source/GeometricObjects/Public/Spatial/FastWinding.h
--- a/Engine/Plugins/Experimental/GeometryProcessing/Source/GeometricObjects/Public/Spatial/FastWinding.h
+++ b/Engine/Plugins/Experimental/GeometryProcessing/Source/GeometricObjects/Public/Spatial/FastWinding.h
@@ -1,5 +1,7 @@
 // Copyright 1998-2019 Epic Games, Inc. All Rights Reserved.
 
+#pragma once
+
 #include "MathUtil.h"
 #include "MeshQueries.h"
 #include "Spatial/MeshAABBTree3.h"
diff --git a/Engine/Plugins/Experimental/GeometryProcessing/Source/GeometricObjects/Public/Util/MeshCaches.h b/Engine/Plugins/Experimental/GeometryProcessing/Source/GeometricObjects/Public/Util/MeshCaches.h
--- a/Engine/Plugins/Experimental/GeometryProcessing/Source/GeometricObjects/Public/Util/MeshCaches.h
+++ b/Engine/Plugins/Experimental/GeometryProcessing/Source/GeometricObjects/Public/Util/MeshCaches.h
@@ -1,5 +1,7 @@
 // Copyright 1998-2019 Epic Games, Inc. All Rights Reserved.
 
+#pragma once
+
 #include "DynamicVector.h"
 
 #include "Async/ParallelFor.h"
diff --git a/Engine/Plugins/FX/HoudiniNiagara/Source/HoudiniNiagaraEditor/Public/HoudiniCSVFactory.h b/Engine/Plugins/FX/HoudiniNiagara/Source/HoudiniNiagaraEditor/Public/HoudiniCSVFactory.h
--- a/Engine/Plugins/FX/HoudiniNiagara/Source/HoudiniNiagaraEditor/Public/HoudiniCSVFactory.h
+++ b/Engine/Plugins/FX/HoudiniNiagara/Source/HoudiniNiagaraEditor/Public/HoudiniCSVFactory.h
@@ -21,6 +21,8 @@
 *
 */
 
+#pragma once
+
 #include "CoreMinimal.h"
 #include "UObject/ObjectMacros.h"
 #include "UObject/UObjectGlobals.h"
diff --git a/Engine/Source/Developer/AssetTools/Private/AssetTypeActions/AssetTypeActions_Texture.cpp b/Engine/Source/Developer/AssetTools/Private/AssetTypeActions/AssetTypeActions_Texture.cpp
--- a/Engine/Source/Developer/AssetTools/Private/AssetTypeActions/AssetTypeActions_Texture.cpp
+++ b/Engine/Source/Developer/AssetTools/Private/AssetTypeActions/AssetTypeActions_Texture.cpp
@@ -1082,8 +1082,8 @@ void FAssetTypeActions_Texture::ConvertVTTexture(TArray<TWeakObjectPtr<UTexture>
 				UE_LOG(LogVirtualTextureConversion, Display, TEXT("Broadcasting to editor."));
 				FEditorDelegates::RefreshEditor.Broadcast();
 				FEditorSupportDelegates::RedrawAllViewports.Broadcast();
-			}*/
-		}
+			}
+		}*/
 #endif
 	}
 }
diff --git a/Engine/Source/Developer/FileUtilities/Public/FileUtilities/ZipArchiveWriter.h b/Engine/Source/Developer/FileUtilities/Public/FileUtilities/ZipArchiveWriter.h
--- a/Engine/Source/Developer/FileUtilities/Public/FileUtilities/ZipArchiveWriter.h
+++ b/Engine/Source/Developer/FileUtilities/Public/FileUtilities/ZipArchiveWriter.h
@@ -1,5 +1,7 @@
 // Copyright 1998-2019 Epic Games, Inc. All Rights Reserved.
 
+#pragma once
+
 #include "CoreMinimal.h"
 #include "GenericPlatform/GenericPlatformFile.h"
 
diff --git a/Engine/Source/Editor/UnrealEd/Private/Bookmarks/BookMark2DTypeActions.h b/Engine/Source/Editor/UnrealEd/Private/Bookmarks/BookMark2DTypeActions.h
--- a/Engine/Source/Editor/UnrealEd/Private/Bookmarks/BookMark2DTypeActions.h
+++ b/Engine/Source/Editor/UnrealEd/Private/Bookmarks/BookMark2DTypeActions.h
@@ -1,5 +1,7 @@
 // Copyright 1998-2019 Epic Games, Inc. All Rights Reserved.
 
+#pragma once
+
 #include "Bookmarks/IBookmarkTypeActions.h"
 #include "EditorViewportClient.h"
 #include "Editor.h"
diff --git a/Engine/Source/Editor/UnrealEd/Public/Bookmarks/BookMarkTypeActions.h b/Engine/Source/Editor/UnrealEd/Public/Bookmarks/BookMarkTypeActions.h
--- a/Engine/Source/Editor/UnrealEd/Public/Bookmarks/BookMarkTypeActions.h
+++ b/Engine/Source/Editor/UnrealEd/Public/Bookmarks/BookMarkTypeActions.h
@@ -1,5 +1,7 @@
 // Copyright 1998-2019 Epic Games, Inc. All Rights Reserved.
 
+#pragma once
+
 #include "Bookmarks/IBookmarkTypeActions.h"
 #include "EditorViewportClient.h"
 #include "Editor.h"
diff --git a/Engine/Source/Programs/IncludeTool/IncludeTool/Program.cs b/Engine/Source/Programs/IncludeTool/IncludeTool/Program.cs
--- a/Engine/Source/Programs/IncludeTool/IncludeTool/Program.cs
+++ b/Engine/Source/Programs/IncludeTool/IncludeTool/Program.cs
@@ -518,7 +518,18 @@ namespace IncludeTool
 				{
 					if ((PreprocessedFile.Flags & SourceFileFlags.Inline) == 0 && (PreprocessedFile.Flags & SourceFileFlags.GeneratedHeader) == 0 && PreprocessedFile.Counterpart == null)
 					{
-						SymbolTable.AddExports(PreprocessedFile);
+						try
+						{
+							SymbolTable.AddExports(PreprocessedFile);
+						}
+						catch(UnbalancedScopeException Ex)
+						{
+							Log.WriteLine("error: {0}", Ex.Message);
+						}
+						catch(Exception Ex)
+						{
+							throw new Exception(String.Format("Error while processing {0}", PreprocessedFile.Location), Ex);
+						}
 					}
 				}
 
diff --git a/Engine/Source/Programs/IncludeTool/IncludeTool/Rules.cs b/Engine/Source/Programs/IncludeTool/IncludeTool/Rules.cs
--- a/Engine/Source/Programs/IncludeTool/IncludeTool/Rules.cs
+++ b/Engine/Source/Programs/IncludeTool/IncludeTool/Rules.cs
@@ -193,7 +193,7 @@ namespace IncludeTool
 			{
 				return true;
 			}
-			if(NormalizedPath.EndsWith("/recastmesh.cpp") || NormalizedPath.EndsWith("/recastfilter.cpp") || NormalizedPath.EndsWith("/recastcontour.cpp") || NormalizedPath.EndsWith("/framepro.h") || NormalizedPath.EndsWith("/framepro.cpp"))
+			if(NormalizedPath.EndsWith("/recastmesh.cpp") || NormalizedPath.EndsWith("/recastfilter.cpp") || NormalizedPath.EndsWith("/recastcontour.cpp") || NormalizedPath.EndsWith("/framepro.h") || NormalizedPath.EndsWith("/framepro.cpp") || NormalizedPath.EndsWith("/sqlite3.inl"))
 			{
 				return true;
 			}
diff --git a/Engine/Source/Programs/IncludeTool/IncludeTool/SymbolTable.cs b/Engine/Source/Programs/IncludeTool/IncludeTool/SymbolTable.cs
--- a/Engine/Source/Programs/IncludeTool/IncludeTool/SymbolTable.cs
+++ b/Engine/Source/Programs/IncludeTool/IncludeTool/SymbolTable.cs
@@ -1,4 +1,4 @@
-﻿// Copyright 1998-2019 Epic Games, Inc. All Rights Reserved.
+// Copyright 1998-2019 Epic Games, Inc. All Rights Reserved.
 
 using IncludeTool.Support;
 using System;
@@ -126,6 +126,14 @@ namespace IncludeTool
 		RequiresDefinition
 	}
 
+	class UnbalancedScopeException : Exception
+	{
+		public UnbalancedScopeException(string Message) 
+			: base(Message)
+		{
+		}
+	}
+
 	/// <summary>
 	/// Registry of symbol names to definitions
 	/// </summary>
@@ -243,7 +251,7 @@ namespace IncludeTool
 					// Make sure they were all consistent
 					if(NewScopes.Any(x => x != NewScopes[0]))
 					{
-						throw new Exception(String.Format("Unbalanced scope depth between conditional block derivations while parsing {0}", File.Location));
+						throw new UnbalancedScopeException(String.Format("Unbalanced scope depth between conditional block derivations while parsing {0}", File.Location));
 					}
 
 					// Update the scope
@@ -285,7 +293,7 @@ namespace IncludeTool
 						}
 						else if(Reader.Current.Text == "}" && --Scope < 0)
 						{
-							throw new Exception(String.Format("Unbalanced '}}' while parsing '{0}'", File.Location));
+							throw new UnbalancedScopeException(String.Format("Unbalanced '}}' while parsing '{0}'", File.Location));
 						}
 
 						// Move to the next token
diff --git a/Engine/Source/Programs/IncludeTool/IncludeTool/TokenReader.cs b/Engine/Source/Programs/IncludeTool/IncludeTool/TokenReader.cs
--- a/Engine/Source/Programs/IncludeTool/IncludeTool/TokenReader.cs
+++ b/Engine/Source/Programs/IncludeTool/IncludeTool/TokenReader.cs
@@ -1,4 +1,4 @@
-﻿// Copyright 1998-2019 Epic Games, Inc. All Rights Reserved.
+// Copyright 1998-2019 Epic Games, Inc. All Rights Reserved.
 
 using IncludeTool.Support;
 using System;
@@ -421,7 +421,7 @@ namespace IncludeTool
 				for (;;)
 				{
 					Character = Text[LineIdx, ColumnIdx];
-					if((Character < 'a' || Character > 'z') && (Character < 'A' || Character > 'Z') && (Character < '0' || Character > '9') && Character != '_' && Character != '$')
+					if((Character < 'a' || Character > 'z') && (Character < 'A' || Character > 'Z') && (Character < '0' || Character > '9') && Character != '_' && Character != '$' && Character != '\'')
 					{
 						if((Character != '+' && Character != '-') || (LastCharacter != 'e' && LastCharacter != 'E'))
 						{
diff --git a/Engine/Source/Runtime/Core/Public/FramePro/FrameProUE4.h b/Engine/Source/Runtime/Core/Public/FramePro/FrameProUE4.h
--- a/Engine/Source/Runtime/Core/Public/FramePro/FrameProUE4.h
+++ b/Engine/Source/Runtime/Core/Public/FramePro/FrameProUE4.h
@@ -22,6 +22,8 @@
 	Add FramePro.cpp to your project to allow FramePro to communicate with your application.
 */
 
+#pragma once
+
 //------------------------------------------------------------------------
 #include "FramePro.h"
 
diff --git a/Engine/Source/Runtime/Core/Public/HAL/LowLevelMemStats.h b/Engine/Source/Runtime/Core/Public/HAL/LowLevelMemStats.h
--- a/Engine/Source/Runtime/Core/Public/HAL/LowLevelMemStats.h
+++ b/Engine/Source/Runtime/Core/Public/HAL/LowLevelMemStats.h
@@ -1,5 +1,7 @@
 // Copyright 1998-2019 Epic Games, Inc. All Rights Reserved.
 
+#pragma once
+
 #include "HAL/LowLevelMemTracker.h"
 #include "Stats/Stats.h"
 
diff --git a/Engine/Source/Runtime/Online/BuildPatchServices/Private/Core/Factory.h b/Engine/Source/Runtime/Online/BuildPatchServices/Private/Core/Factory.h
--- a/Engine/Source/Runtime/Online/BuildPatchServices/Private/Core/Factory.h
+++ b/Engine/Source/Runtime/Online/BuildPatchServices/Private/Core/Factory.h
@@ -1,5 +1,7 @@
 // Copyright 1998-2019 Epic Games, Inc. All Rights Reserved.
 
+#pragma once
+
 namespace BuildPatchServices
 {
 	/**
diff --git a/Engine/Source/Runtime/Renderer/Private/IndirectLightRendering.h b/Engine/Source/Runtime/Renderer/Private/IndirectLightRendering.h
--- a/Engine/Source/Runtime/Renderer/Private/IndirectLightRendering.h
+++ b/Engine/Source/Runtime/Renderer/Private/IndirectLightRendering.h
@@ -1,4 +1,6 @@
 // Copyright 1998-2019 Epic Games, Inc. All Rights Reserved.
 
+#pragma once
+
 #include "ScenePrivate.h"
 
