Last batch of IncludeTool fixes (hopefully).

#rnx

diff --git a/Engine/Plugins/VirtualProduction/Takes/Source/TakeRecorderSources/Private/TakeRecorderLevelSequenceSource.cpp b/Engine/Plugins/VirtualProduction/Takes/Source/TakeRecorderSources/Private/TakeRecorderLevelSequenceSource.cpp
--- a/Engine/Plugins/VirtualProduction/Takes/Source/TakeRecorderSources/Private/TakeRecorderLevelSequenceSource.cpp
+++ b/Engine/Plugins/VirtualProduction/Takes/Source/TakeRecorderSources/Private/TakeRecorderLevelSequenceSource.cpp
@@ -1,7 +1,7 @@
 // Copyright 1998-2019 Epic Games, Inc. All Rights Reserved.
 
 #include "TakeRecorderLevelSequenceSource.h"
-#include "TakesCoreFwd.h"
+#include "TakesCoreLog.h"
 #include "TakesUtils.h"
 #include "Editor.h"
 #include "GameFramework/Actor.h"
diff --git a/Engine/Plugins/VirtualProduction/Takes/Source/TakeRecorderSources/Private/TakeRecorderNearbySpawnedActorSource.cpp b/Engine/Plugins/VirtualProduction/Takes/Source/TakeRecorderSources/Private/TakeRecorderNearbySpawnedActorSource.cpp
--- a/Engine/Plugins/VirtualProduction/Takes/Source/TakeRecorderSources/Private/TakeRecorderNearbySpawnedActorSource.cpp
+++ b/Engine/Plugins/VirtualProduction/Takes/Source/TakeRecorderSources/Private/TakeRecorderNearbySpawnedActorSource.cpp
@@ -7,7 +7,7 @@
 #include "TakeRecorderActorSource.h"
 #include "TakeRecorderSourcesUtils.h"
 #include "TakesUtils.h"
-#include "TakesCoreFwd.h"
+#include "TakesCoreLog.h"
 
 #include "LevelSequence.h"
 #include "Engine/Engine.h"
diff --git a/Engine/Plugins/VirtualProduction/Takes/Source/TakeRecorderSources/Private/TakeRecorderSourcesModule.cpp b/Engine/Plugins/VirtualProduction/Takes/Source/TakeRecorderSources/Private/TakeRecorderSourcesModule.cpp
--- a/Engine/Plugins/VirtualProduction/Takes/Source/TakeRecorderSources/Private/TakeRecorderSourcesModule.cpp
+++ b/Engine/Plugins/VirtualProduction/Takes/Source/TakeRecorderSources/Private/TakeRecorderSourcesModule.cpp
@@ -8,7 +8,7 @@
 #include "Recorder/TakeRecorder.h"
 #include "Recorder/TakeRecorderParameters.h"
 #include "TakesCoreBlueprintLibrary.h"
-#include "TakesCoreFwd.h"
+#include "TakesCoreLog.h"
 #include "TakeMetaData.h"
 #include "TakeRecorderActorSource.h"
 #include "TakeRecorderSources.h"
diff --git a/Engine/Plugins/VirtualProduction/Takes/Source/TakeTrackRecorders/Private/MovieSceneVisibilityTrackRecorder.cpp b/Engine/Plugins/VirtualProduction/Takes/Source/TakeTrackRecorders/Private/MovieSceneVisibilityTrackRecorder.cpp
--- a/Engine/Plugins/VirtualProduction/Takes/Source/TakeTrackRecorders/Private/MovieSceneVisibilityTrackRecorder.cpp
+++ b/Engine/Plugins/VirtualProduction/Takes/Source/TakeTrackRecorders/Private/MovieSceneVisibilityTrackRecorder.cpp
@@ -1,7 +1,7 @@
 // Copyright 1998-2019 Epic Games, Inc. All Rights Reserved.
 
 #include "TrackRecorders/MovieSceneVisibilityTrackRecorder.h"
-#include "TakesCoreFwd.h"
+#include "TakesCoreLog.h"
 #include "Components/SceneComponent.h"
 #include "GameFramework/Actor.h"
 #include "KeyParams.h"
diff --git a/Engine/Plugins/VirtualProduction/Takes/Source/TakeTrackRecorders/Public/TrackRecorders/MovieScene3DTransformTrackRecorder.h b/Engine/Plugins/VirtualProduction/Takes/Source/TakeTrackRecorders/Public/TrackRecorders/MovieScene3DTransformTrackRecorder.h
--- a/Engine/Plugins/VirtualProduction/Takes/Source/TakeTrackRecorders/Public/TrackRecorders/MovieScene3DTransformTrackRecorder.h
+++ b/Engine/Plugins/VirtualProduction/Takes/Source/TakeTrackRecorders/Public/TrackRecorders/MovieScene3DTransformTrackRecorder.h
@@ -9,7 +9,7 @@
 #include "Curves/RichCurve.h"
 #include "Serializers/MovieSceneTransformSerialization.h"
 #include "Sections/MovieScene3DTransformSection.h"
-#include "TakesCoreFwd.h"
+#include "TakesCoreLog.h"
 #include "MovieScene3DTransformTrackRecorder.generated.h"
 
 class FMovieScene3DTransformTrackRecorder;
diff --git a/Engine/Plugins/VirtualProduction/Takes/Source/TakeTrackRecorders/Public/TrackRecorders/MovieSceneTrackPropertyRecorder.h b/Engine/Plugins/VirtualProduction/Takes/Source/TakeTrackRecorders/Public/TrackRecorders/MovieSceneTrackPropertyRecorder.h
--- a/Engine/Plugins/VirtualProduction/Takes/Source/TakeTrackRecorders/Public/TrackRecorders/MovieSceneTrackPropertyRecorder.h
+++ b/Engine/Plugins/VirtualProduction/Takes/Source/TakeTrackRecorders/Public/TrackRecorders/MovieSceneTrackPropertyRecorder.h
@@ -7,7 +7,7 @@
 #include "MovieSceneSection.h"
 #include "MovieSceneTrack.h"
 #include "MovieScene.h"
-#include "TakesCoreFwd.h"
+#include "TakesCoreLog.h"
 #include "Channels/MovieSceneChannel.h"
 #include "Channels/MovieSceneChannelData.h"
 #include "Channels/MovieSceneChannelProxy.h"
diff --git a/Engine/Plugins/VirtualProduction/Takes/Source/TakesCore/Private/TakeData.h b/Engine/Plugins/VirtualProduction/Takes/Source/TakesCore/Private/TakeData.h
--- a/Engine/Plugins/VirtualProduction/Takes/Source/TakesCore/Private/TakeData.h
+++ b/Engine/Plugins/VirtualProduction/Takes/Source/TakesCore/Private/TakeData.h
@@ -2,7 +2,7 @@
 
 #pragma once
 
-#include "TakesCoreFwd.h"
+#include "TakesCoreLog.h"
 #include "TakeMetaData.h"
 #include "TakesCoreBlueprintLibrary.h"
 
diff --git a/Engine/Plugins/VirtualProduction/Takes/Source/TakesCore/Private/TakeRecorderSources.cpp b/Engine/Plugins/VirtualProduction/Takes/Source/TakesCore/Private/TakeRecorderSources.cpp
--- a/Engine/Plugins/VirtualProduction/Takes/Source/TakesCore/Private/TakeRecorderSources.cpp
+++ b/Engine/Plugins/VirtualProduction/Takes/Source/TakesCore/Private/TakeRecorderSources.cpp
@@ -2,7 +2,7 @@
 
 #include "TakeRecorderSources.h"
 #include "TakeRecorderSource.h"
-#include "TakesCoreFwd.h"
+#include "TakesCoreLog.h"
 #include "LevelSequence.h"
 #include "TakeMetaData.h"
 #include "AssetRegistryModule.h"
diff --git a/Engine/Plugins/VirtualProduction/Takes/Source/TakesCore/Private/TakesCoreModule.cpp b/Engine/Plugins/VirtualProduction/Takes/Source/TakesCore/Private/TakesCoreModule.cpp
--- a/Engine/Plugins/VirtualProduction/Takes/Source/TakesCore/Private/TakesCoreModule.cpp
+++ b/Engine/Plugins/VirtualProduction/Takes/Source/TakesCore/Private/TakesCoreModule.cpp
@@ -1,6 +1,6 @@
 // Copyright 1998-2019 Epic Games, Inc. All Rights Reserved.
 
-#include "TakesCoreFwd.h"
+#include "TakesCoreLog.h"
 #include "TakeMetaData.h"
 #include "TakeData.h"
 
diff --git a/Engine/Plugins/VirtualProduction/Takes/Source/TakesCore/Public/TakesCoreFwd.h b/Engine/Plugins/VirtualProduction/Takes/Source/TakesCore/Public/TakesCoreFwd.h
deleted file mode 100644
diff --git a/Engine/Plugins/VirtualProduction/Takes/Source/TakesCore/Public/TakesCoreLog.h b/Engine/Plugins/VirtualProduction/Takes/Source/TakesCore/Public/TakesCoreLog.h
new file mode 100644
--- /dev/null
+++ b/Engine/Plugins/VirtualProduction/Takes/Source/TakesCore/Public/TakesCoreLog.h
@@ -0,0 +1,7 @@
+// Copyright 1998-2019 Epic Games, Inc. All Rights Reserved.
+
+#pragma once
+
+#include "CoreMinimal.h"
+
+TAKESCORE_API DECLARE_LOG_CATEGORY_EXTERN(LogTakesCore, Log, All);
\ No newline at end of file
diff --git a/Engine/Plugins/VirtualProduction/Takes/Source/TakesCore/Public/TakesUtils.h b/Engine/Plugins/VirtualProduction/Takes/Source/TakesCore/Public/TakesUtils.h
--- a/Engine/Plugins/VirtualProduction/Takes/Source/TakesCore/Public/TakesUtils.h
+++ b/Engine/Plugins/VirtualProduction/Takes/Source/TakesCore/Public/TakesUtils.h
@@ -8,7 +8,7 @@
 #include "IAssetRegistry.h"
 #include "AssetRegistryModule.h"
 #include "UObject/UObjectGlobals.h"
-#include "TakesCoreFwd.h"
+#include "TakesCoreLog.h"
 
 class UWorld;
 class ULevelSequence;
diff --git a/Engine/Source/Programs/IncludeTool/IncludeTool/Rules.cs b/Engine/Source/Programs/IncludeTool/IncludeTool/Rules.cs
--- a/Engine/Source/Programs/IncludeTool/IncludeTool/Rules.cs
+++ b/Engine/Source/Programs/IncludeTool/IncludeTool/Rules.cs
@@ -152,6 +152,10 @@ namespace IncludeTool
 			// XboxOne
 			"EtwPlus.h",
 
+			// Vorbis
+			"vorbis_stream_encoder.h",
+			"integral_types.h",
+
 			// UHT
 			"StructSpecifiers.def",
 			"FunctionSpecifiers.def",
@@ -193,7 +197,7 @@ namespace IncludeTool
 			{
 				return true;
 			}
-			if(NormalizedPath.EndsWith("/recastmesh.cpp") || NormalizedPath.EndsWith("/recastfilter.cpp") || NormalizedPath.EndsWith("/recastcontour.cpp") || NormalizedPath.EndsWith("/framepro.h") || NormalizedPath.EndsWith("/framepro.cpp") || NormalizedPath.EndsWith("/sqlite3.h") || NormalizedPath.EndsWith("/sqlite3.inl") || NormalizedPath.EndsWith("/vorbis_stream_encoder.h"))
+			if(NormalizedPath.EndsWith("/recastmesh.cpp") || NormalizedPath.EndsWith("/recastfilter.cpp") || NormalizedPath.EndsWith("/recastcontour.cpp") || NormalizedPath.EndsWith("/framepro.h") || NormalizedPath.EndsWith("/framepro.cpp") || NormalizedPath.EndsWith("/sqlite3.h") || NormalizedPath.EndsWith("/sqlite3.inl") || NormalizedPath.EndsWith("/vorbis_stream_encoder.h") || NormalizedPath.EndsWith("/integral_types.h"))
 			{
 				return true;
 			}
diff --git a/Engine/Source/Programs/IncludeTool/IncludeTool/SymbolTable.cs b/Engine/Source/Programs/IncludeTool/IncludeTool/SymbolTable.cs
--- a/Engine/Source/Programs/IncludeTool/IncludeTool/SymbolTable.cs
+++ b/Engine/Source/Programs/IncludeTool/IncludeTool/SymbolTable.cs
@@ -710,7 +710,8 @@ namespace IncludeTool
 				{
 					if(!ReadClassOrStructForwardDeclaration(Reader, HeaderFile, SymbolToHeader, Log) 
 						&& !ReadTemplateClassOrStructForwardDeclaration(Reader, HeaderFile, SymbolToHeader, Log)
-						&& !ReadEnumClassForwardDeclaration(Reader, HeaderFile, SymbolToHeader, Log))
+						&& !ReadEnumClassForwardDeclaration(Reader, HeaderFile, SymbolToHeader, Log)
+						&& !SkipAliasDeclaration(Reader))
 					{
 						Log.WriteLine("{0}({1}): error: invalid forward declaration - '{2}'", HeaderFile.Location, Reader.CurrentLine + 1, HeaderFile.Text[Reader.CurrentLine]);
 						return false;
@@ -896,6 +897,46 @@ namespace IncludeTool
 			return bMoveNext;
 		}
 
+		/// <summary>
+		/// Skip over the declaration of an alias
+		/// </summary>
+		/// <param name="OriginalReader">The current token reader</param>
+		/// <returns>True if a forward declaration was read, false otherwise</returns>
+		bool SkipAliasDeclaration(TokenReader OriginalReader)
+		{
+			TokenReader Reader = new TokenReader(OriginalReader);
+
+			if(Reader.Current.Text == "template")
+			{
+				if (!Reader.MoveNext(TokenReaderContext.IgnoreNewlines))
+				{
+					return false;
+				}
+
+				bool bMoveNext = true;
+				if (!SkipTemplateArguments(Reader, ref bMoveNext) || !bMoveNext)
+				{
+					return false;
+				}
+			}
+
+			if(Reader.Current.Text != "using")
+			{
+				return false;
+			}
+
+			while(Reader.MoveNext(TokenReaderContext.IgnoreNewlines))
+			{
+				if(Reader.Current.Text == ";")
+				{
+					break;
+				}
+			}
+
+			OriginalReader.Set(Reader);
+			return true;
+		}
+
 		/// <summary>
 		/// Adds a forward declaration to a map
 		/// </summary>
diff --git a/Engine/Source/Runtime/CoreUObject/Public/UObject/UObjectBaseUtility.h b/Engine/Source/Runtime/CoreUObject/Public/UObject/UObjectBaseUtility.h
--- a/Engine/Source/Runtime/CoreUObject/Public/UObject/UObjectBaseUtility.h
+++ b/Engine/Source/Runtime/CoreUObject/Public/UObject/UObjectBaseUtility.h
@@ -755,8 +755,9 @@ public:
 
 #elif ENABLE_STATNAMEDEVENTS
 
-struct FScopeCycleCounterUObject
+class FScopeCycleCounterUObject
 {
+public:
 	FScopeCycleCounter ScopeCycleCounter;
 	FORCEINLINE_STATS FScopeCycleCounterUObject(const UObjectBaseUtility *Object)
 	: ScopeCycleCounter(Object ? Object->GetStatID().StatString : nullptr)
@@ -798,8 +799,9 @@ public:
 #define SCOPE_CYCLE_UOBJECT(Name, Object) \
 	FScopeCycleCounterUObject ObjCycleCount_##Name(Object);
 #else
-struct FScopeCycleCounterUObject
+class FScopeCycleCounterUObject
 {
+public:
 	FORCEINLINE_STATS FScopeCycleCounterUObject(const UObjectBaseUtility *Object)
 	{
 	}
