IncludeTool: Add a LineBasedTextWriter class, to allow outputting of log data line-by-line. The previous character-based writer could cause garbled output when multiple threads try to write at once.

diff --git a/Engine/Source/Programs/IncludeTool/IncludeTool/IncludeTool.csproj b/Engine/Source/Programs/IncludeTool/IncludeTool/IncludeTool.csproj
--- a/Engine/Source/Programs/IncludeTool/IncludeTool/IncludeTool.csproj
+++ b/Engine/Source/Programs/IncludeTool/IncludeTool/IncludeTool.csproj
@@ -116,6 +116,8 @@
     <Compile Include="Support\HashList.cs" />
     <Compile Include="Support\JsonObject.cs" />
     <Compile Include="Support\JsonWriter.cs" />
+    <Compile Include="Support\LineBasedTextWriter.cs" />
+    <Compile Include="Support\LineBasedTextWriterWrapper.cs" />
     <Compile Include="Support\LogWriter.cs" />
     <Compile Include="Support\ManagedTask.cs" />
     <Compile Include="Support\MultiValueDictionary.cs" />
diff --git a/Engine/Source/Programs/IncludeTool/IncludeTool/OutputFile.cs b/Engine/Source/Programs/IncludeTool/IncludeTool/OutputFile.cs
--- a/Engine/Source/Programs/IncludeTool/IncludeTool/OutputFile.cs
+++ b/Engine/Source/Programs/IncludeTool/IncludeTool/OutputFile.cs
@@ -181,7 +181,7 @@ namespace IncludeTool
 		/// <param name="CppFiles">Input files to optimize</param>
 		/// <param name="Log">Writer for log messages</param>
 		/// <returns>Array of output files</returns>
-		public static void PrepareFilesForOutput(IEnumerable<SourceFile> CppFiles, Dictionary<SourceFile, SourceFile> CppFileToHeaderFile, Dictionary<Symbol, SourceFile> FwdSymbolToInputHeader, bool bMakeStandalone, bool bUseOriginalIncludes, TextWriter Log)
+		public static void PrepareFilesForOutput(IEnumerable<SourceFile> CppFiles, Dictionary<SourceFile, SourceFile> CppFileToHeaderFile, Dictionary<Symbol, SourceFile> FwdSymbolToInputHeader, bool bMakeStandalone, bool bUseOriginalIncludes, LineBasedTextWriter Log)
 		{
 			// Cache of all the created output files
 			Dictionary<SourceFile, OutputFile> OutputFileLookup = new Dictionary<SourceFile,OutputFile>();
@@ -240,7 +240,7 @@ namespace IncludeTool
 		/// <param name="OutputFiles">List of output files</param>
 		/// <param name="OutputFileLookup">Mapping from source file to output file</param>
 		/// <returns>The new or existing output file</returns>
-		static OutputFile FindOrCreateOutputFile(List<SourceFile> InputFileStack, Dictionary<SourceFile, SourceFile> CppFileToHeaderFile, HashList<OutputFile> PreviousFiles, Dictionary<SourceFile, OutputFile> OutputFileLookup, Dictionary<Symbol, OutputFile> FwdSymbolToHeader, bool bMakeStandalone, bool bUseOriginalIncludes, TextWriter Log)
+		static OutputFile FindOrCreateOutputFile(List<SourceFile> InputFileStack, Dictionary<SourceFile, SourceFile> CppFileToHeaderFile, HashList<OutputFile> PreviousFiles, Dictionary<SourceFile, OutputFile> OutputFileLookup, Dictionary<Symbol, OutputFile> FwdSymbolToHeader, bool bMakeStandalone, bool bUseOriginalIncludes, LineBasedTextWriter Log)
 		{
 			// Get the file at the top of the stack
 			SourceFile InputFile = InputFileStack[InputFileStack.Count - 1];
@@ -334,7 +334,7 @@ namespace IncludeTool
 		/// <param name="FwdSymbolToHeader"></param>
 		/// <param name="bMakeStandalone">Whether to make this output file standalone</param>
 		/// <param name="Log">Writer for log messages</param>
-		public static OutputFile CreateOptimizedOutputFile(SourceFile InputFile, OutputFile HeaderFile, List<OutputFile> PreviousFiles, List<OutputFileInclude> Includes, List<SourceFile> InputFileStack, Dictionary<Symbol, OutputFile> FwdSymbolToHeader, bool bMakeStandalone, TextWriter Log)
+		public static OutputFile CreateOptimizedOutputFile(SourceFile InputFile, OutputFile HeaderFile, List<OutputFile> PreviousFiles, List<OutputFileInclude> Includes, List<SourceFile> InputFileStack, Dictionary<Symbol, OutputFile> FwdSymbolToHeader, bool bMakeStandalone, LineBasedTextWriter Log)
 		{
 			Debug.Assert(HeaderFile == null || (InputFile.Flags & SourceFileFlags.TranslationUnit) != 0);
 
@@ -590,7 +590,7 @@ namespace IncludeTool
 		/// <param name="InputFile">The input file that this output file corresponds to</param>
 		/// <param name="Includes">The active set of includes parsed for this file</param>
 		/// <param name="Log">Writer for log messages</param>
-		public static OutputFile CreatePassthroughOutputFile(SourceFile InputFile, List<OutputFileInclude> Includes, TextWriter Log)
+		public static OutputFile CreatePassthroughOutputFile(SourceFile InputFile, List<OutputFileInclude> Includes, LineBasedTextWriter Log)
 		{
 			// Write the state
 			InputFile.LogVerbose("InputFile={0}", InputFile.Location.FullName);
diff --git a/Engine/Source/Programs/IncludeTool/IncludeTool/Preprocessor.cs b/Engine/Source/Programs/IncludeTool/IncludeTool/Preprocessor.cs
--- a/Engine/Source/Programs/IncludeTool/IncludeTool/Preprocessor.cs
+++ b/Engine/Source/Programs/IncludeTool/IncludeTool/Preprocessor.cs
@@ -1399,7 +1399,7 @@ namespace IncludeTool
 				// Spawn clang and parse the output
 				using (StreamWriter Writer = new StreamWriter(PreludeFile.FullName))
 				{
-					Utility.Run(CompileEnvironment.Compiler, String.Join(" ", CompileEnvironment.Options.Where(x => x.Name == "-target").Select(x => x.ToString())) + " -dM -E -x c++ " + InputFile.FullName, WorkingDir, Writer);
+					Utility.Run(CompileEnvironment.Compiler, String.Join(" ", CompileEnvironment.Options.Where(x => x.Name == "-target").Select(x => x.ToString())) + " -dM -E -x c++ " + InputFile.FullName, WorkingDir, new LineBasedTextWriterWrapper(Writer));
 				}
 			}
 			else if(CompileEnvironment.CompilerType == CompilerType.VisualC)
@@ -1424,7 +1424,7 @@ namespace IncludeTool
 				// Invoke the compiler and capture the output into a string
 				StringWriter OutputWriter = new StringWriter();
 				string FullCommandLine = String.Format("{0} /Zs /TP {1}", CompileEnvironment.GetCommandLine(), InputFile.FullName);
-				Utility.Run(CompileEnvironment.Compiler, FullCommandLine, WorkingDir, OutputWriter);
+				Utility.Run(CompileEnvironment.Compiler, FullCommandLine, WorkingDir, new LineBasedTextWriterWrapper(OutputWriter));
 
 				// Filter the output
 				using (StreamWriter Writer = new StreamWriter(PreludeFile.FullName))
@@ -2014,7 +2014,7 @@ namespace IncludeTool
 		/// <param name="PreprocessedDir">The output directory for the preprocessed files</param>
 		/// <param name="FilePairs">List of files and their compile environment</param>
 		/// <param name="Log">Writer for log output</param>
-		public static void WritePreprocessedFiles(DirectoryReference WorkingDir, DirectoryReference PreprocessedDir, IEnumerable<KeyValuePair<FileReference, CompileEnvironment>> FilePairs, TextWriter Log)
+		public static void WritePreprocessedFiles(DirectoryReference WorkingDir, DirectoryReference PreprocessedDir, IEnumerable<KeyValuePair<FileReference, CompileEnvironment>> FilePairs, LineBasedTextWriter Log)
 		{
 			PreprocessedDir.CreateDirectory();
 
@@ -2058,7 +2058,7 @@ namespace IncludeTool
 		/// <param name="ToolDir">Output directory for files preprocessed by this tool</param>
 		/// <param name="FileCache">Cache of file locations to their source file</param>
 		/// <param name="Log">Output log writer</param>
-		static void WritePreprocessedFile(FileReference InputFile, CompileEnvironment Environment, FileReference PreludeFile, DirectoryReference WorkingDir, DirectoryReference CompilerDir, DirectoryReference CompilerNormalizedDir, DirectoryReference ToolDir, TextWriter Log)
+		static void WritePreprocessedFile(FileReference InputFile, CompileEnvironment Environment, FileReference PreludeFile, DirectoryReference WorkingDir, DirectoryReference CompilerDir, DirectoryReference CompilerNormalizedDir, DirectoryReference ToolDir, LineBasedTextWriter Log)
 		{
 			BufferedTextWriter BufferedLog = new BufferedTextWriter();
 
@@ -2090,7 +2090,7 @@ namespace IncludeTool
 		/// <param name="CompilerFile">Path to the preprocessed file generated by the compiler</param>
 		/// <param name="NormalizedCompilerFile">Output file for the normalized compiler file</param>
 		/// <param name="Log"></param>
-		static int NormalizePreprocessedFile(FileReference ToolFile, FileReference CompilerFile, FileReference NormalizedCompilerFile, TextWriter Log)
+		static int NormalizePreprocessedFile(FileReference ToolFile, FileReference CompilerFile, FileReference NormalizedCompilerFile, LineBasedTextWriter Log)
 		{
 			// Read the tool file
 			string[] ToolLines = File.ReadAllLines(ToolFile.FullName);
diff --git a/Engine/Source/Programs/IncludeTool/IncludeTool/Program.cs b/Engine/Source/Programs/IncludeTool/IncludeTool/Program.cs
--- a/Engine/Source/Programs/IncludeTool/IncludeTool/Program.cs
+++ b/Engine/Source/Programs/IncludeTool/IncludeTool/Program.cs
@@ -860,7 +860,7 @@ namespace IncludeTool
 		/// <param name="Target">The target to compile (eg. UE4Editor Win64 Development)</param>
 		/// <param name="TaskListFile">The output file for the task list</param>
 		/// <param name="Log">Log to output messages to</param>
-		static void GenerateTaskList(DirectoryReference RootDir, string Target, TargetPlatform Platform, TargetConfiguration Configuration, bool Precompile, FileReference TaskListFile, TextWriter Log)
+		static void GenerateTaskList(DirectoryReference RootDir, string Target, TargetPlatform Platform, TargetConfiguration Configuration, bool Precompile, FileReference TaskListFile, LineBasedTextWriter Log)
 		{
 			DirectoryReference IntermediateDir = DirectoryReference.Combine(RootDir, "Engine", "Intermediate", "Build");
 			if(IntermediateDir.Exists())
@@ -917,7 +917,7 @@ namespace IncludeTool
 		/// </summary>
 		/// <param name="WorkingDir">The working directory for fragments</param>
 		/// <param name="FilesToCompile">Array of files to compile, with their corresponding compile environment</param>
-		static void PreprocessFiles(DirectoryReference WorkingDir, KeyValuePair<FileReference, CompileEnvironment>[] FilesToCompile, IEnumerable<DirectoryReference> ExtraSystemIncludePaths, TextWriter Log)
+		static void PreprocessFiles(DirectoryReference WorkingDir, KeyValuePair<FileReference, CompileEnvironment>[] FilesToCompile, IEnumerable<DirectoryReference> ExtraSystemIncludePaths, LineBasedTextWriter Log)
 		{
 			if(FilesToCompile.Length > 0)
 			{
@@ -939,7 +939,7 @@ namespace IncludeTool
 		/// <param name="CompileEnvironment">Compile environment for the given file</param>
 		/// <param name="FileToActiveMarkup">Map from source file to bit array of the markup blocks that are active. Used to verify all preprocessed files parse the same way.</param>
 		/// <param name="Log">Writer for output messages</param>
-		static void PreprocessFile(FileReference Location, FileReference PreludeLocation, CompileEnvironment CompileEnvironment, IEnumerable<DirectoryReference> ExtraSystemIncludePaths, ConcurrentDictionary<SourceFile, MarkupState> FileToActiveMarkup, TextWriter Log)
+		static void PreprocessFile(FileReference Location, FileReference PreludeLocation, CompileEnvironment CompileEnvironment, IEnumerable<DirectoryReference> ExtraSystemIncludePaths, ConcurrentDictionary<SourceFile, MarkupState> FileToActiveMarkup, LineBasedTextWriter Log)
 		{
 			// Get the initial file
 			PreprocessorFile InitialFile = new PreprocessorFile(Location.FullName, Workspace.GetFile(Location));
@@ -964,7 +964,7 @@ namespace IncludeTool
 		/// <param name="IncludedFile">The file to include</param>
 		/// <param name="FileToActiveMarkup">Map from source file to bit array of the markup blocks that are active. Used to verify all preprocessed files parse the same way.</param>
 		/// <param name="Log">Writer for output messages</param>
-		static void PreprocessIncludedFile(Preprocessor Preprocessor, PreprocessorFile IncludedFile, ConcurrentDictionary<SourceFile, MarkupState> FileToActiveMarkup, TextWriter Log)
+		static void PreprocessIncludedFile(Preprocessor Preprocessor, PreprocessorFile IncludedFile, ConcurrentDictionary<SourceFile, MarkupState> FileToActiveMarkup, LineBasedTextWriter Log)
 		{
 			if(Preprocessor.PushFile(IncludedFile))
 			{
@@ -1257,7 +1257,7 @@ namespace IncludeTool
 		/// <param name="bMultiProcess">Whether to spawn child processes to execute the search. Useful with XGE.</param>
 		/// <param name="Log">Writer for log files</param>
 		/// <returns>True on success</returns>
-		public static bool ProcessSequenceTree(SequenceProbeType Type, Dictionary<SourceFile, CompileEnvironment> SourceFileToCompileEnvironment, DirectoryReference IntermediateDir, IEnumerable<DirectoryReference> ExtraSystemIncludePaths, int MaxParallel, int ShardIdx, int NumShards, double TimeLimit, TextWriter Log)
+		public static bool ProcessSequenceTree(SequenceProbeType Type, Dictionary<SourceFile, CompileEnvironment> SourceFileToCompileEnvironment, DirectoryReference IntermediateDir, IEnumerable<DirectoryReference> ExtraSystemIncludePaths, int MaxParallel, int ShardIdx, int NumShards, double TimeLimit, LineBasedTextWriter Log)
 		{
 			Stopwatch Timer = Stopwatch.StartNew();
 			Log.WriteLine("Creating probes...");
@@ -1493,7 +1493,7 @@ namespace IncludeTool
 		/// <param name="bWithP4">If true, read-only files will be checked out of Perforce</param>
 		/// <param name="Log">Log writer</param>
 		/// <returns>True if the files were written successfully</returns>
-		static bool WriteOptimizedFiles(DirectoryReference InputDir, DirectoryReference OutputDir, HashList<DirectoryReference> PublicIncludePaths, Dictionary<BuildModule, HashList<DirectoryReference>> PrivateIncludePaths, HashList<DirectoryReference> SystemIncludePaths, IEnumerable<SourceFile> OutputFiles, bool bWithP4, bool bWithVerboseLog, bool bRemoveForwardDeclarations, TextWriter Log)
+		static bool WriteOptimizedFiles(DirectoryReference InputDir, DirectoryReference OutputDir, HashList<DirectoryReference> PublicIncludePaths, Dictionary<BuildModule, HashList<DirectoryReference>> PrivateIncludePaths, HashList<DirectoryReference> SystemIncludePaths, IEnumerable<SourceFile> OutputFiles, bool bWithP4, bool bWithVerboseLog, bool bRemoveForwardDeclarations, LineBasedTextWriter Log)
 		{
 			// Find all the output files that need to be written
 			Dictionary<FileReference, string> OutputFileContents = new Dictionary<FileReference, string>();
@@ -1531,7 +1531,7 @@ namespace IncludeTool
 		/// <param name="bWithP4">If true, read-only files will be checked out of Perforce</param>
 		/// <param name="Log">Log writer</param>
 		/// <returns>True if the files were written successfully</returns>
-		static bool WriteOutputFiles(DirectoryReference OutputDir, Dictionary<FileReference, string> OutputFileContents, bool bWithP4, TextWriter Log)
+		static bool WriteOutputFiles(DirectoryReference OutputDir, Dictionary<FileReference, string> OutputFileContents, bool bWithP4, LineBasedTextWriter Log)
 		{
 			// Create the output directory
 			OutputDir.CreateDirectory();
diff --git a/Engine/Source/Programs/IncludeTool/IncludeTool/Reports/ComplexityReport.cs b/Engine/Source/Programs/IncludeTool/IncludeTool/Reports/ComplexityReport.cs
--- a/Engine/Source/Programs/IncludeTool/IncludeTool/Reports/ComplexityReport.cs
+++ b/Engine/Source/Programs/IncludeTool/IncludeTool/Reports/ComplexityReport.cs
@@ -1,4 +1,4 @@
-﻿// Copyright 1998-2019 Epic Games, Inc. All Rights Reserved.
+// Copyright 1998-2019 Epic Games, Inc. All Rights Reserved.
 
 using IncludeTool.Support;
 using System;
@@ -43,7 +43,7 @@ namespace IncludeTool.Reports
 		/// <param name="Files">The files to include in the report</param>
 		/// <param name="ReportFileLocation">Output file for the report</param>
 		/// <param name="Log">Writer for log output</param>
-		public static void Generate(FileReference ReportFileLocation, IEnumerable<SourceFile> Files, TextWriter Log)
+		public static void Generate(FileReference ReportFileLocation, IEnumerable<SourceFile> Files, LineBasedTextWriter Log)
 		{
 			Log.WriteLine("Writing {0}...", ReportFileLocation.FullName);
 
diff --git a/Engine/Source/Programs/IncludeTool/IncludeTool/Reports/FragmentTimingReport.cs b/Engine/Source/Programs/IncludeTool/IncludeTool/Reports/FragmentTimingReport.cs
--- a/Engine/Source/Programs/IncludeTool/IncludeTool/Reports/FragmentTimingReport.cs
+++ b/Engine/Source/Programs/IncludeTool/IncludeTool/Reports/FragmentTimingReport.cs
@@ -1,4 +1,4 @@
-﻿// Copyright 1998-2019 Epic Games, Inc. All Rights Reserved.
+// Copyright 1998-2019 Epic Games, Inc. All Rights Reserved.
 
 using IncludeTool.Support;
 using System;
@@ -254,7 +254,7 @@ namespace IncludeTool.Reports
 		/// <param name="NumSamples">Number of samples to take</param>
 		/// <param name="MaxParallel">Maximum number of tasks to run in parallel.</param>
 		/// <param name="Log">Log writer</param>
-		public static void Generate(DirectoryReference InputDir, DirectoryReference WorkingDir, Dictionary<SourceFile, CompileEnvironment> FileToCompileEnvironment, int NumSamples, int Shard, int NumShards, int MaxParallel, TextWriter Log)
+		public static void Generate(DirectoryReference InputDir, DirectoryReference WorkingDir, Dictionary<SourceFile, CompileEnvironment> FileToCompileEnvironment, int NumSamples, int Shard, int NumShards, int MaxParallel, LineBasedTextWriter Log)
 		{
 			Stopwatch Timer = Stopwatch.StartNew();
 
diff --git a/Engine/Source/Programs/IncludeTool/IncludeTool/Reports/PchReport.cs b/Engine/Source/Programs/IncludeTool/IncludeTool/Reports/PchReport.cs
--- a/Engine/Source/Programs/IncludeTool/IncludeTool/Reports/PchReport.cs
+++ b/Engine/Source/Programs/IncludeTool/IncludeTool/Reports/PchReport.cs
@@ -1,4 +1,4 @@
-﻿// Copyright 1998-2019 Epic Games, Inc. All Rights Reserved.
+// Copyright 1998-2019 Epic Games, Inc. All Rights Reserved.
 
 using IncludeTool.Support;
 using System;
@@ -36,7 +36,7 @@ namespace IncludeTool.Reports
 		/// <param name="Target"></param>
 		/// <param name="SourceFileToCompileEnvironment"></param>
 		/// <param name="Log">Writer for log output</param>
-		public static void Generate(FileReference ReportFileLocation, DirectoryReference InputDir, BuildTarget Target, Dictionary<SourceFile, CompileEnvironment> SourceFileToCompileEnvironment, TextWriter Log)
+		public static void Generate(FileReference ReportFileLocation, DirectoryReference InputDir, BuildTarget Target, Dictionary<SourceFile, CompileEnvironment> SourceFileToCompileEnvironment, LineBasedTextWriter Log)
 		{
 			Log.WriteLine("Writing {0}...", ReportFileLocation.FullName);
 
@@ -72,7 +72,7 @@ namespace IncludeTool.Reports
 		/// <param name="OutputFileContents">Dictionary which receives the files to output</param>
 		/// <param name="MinRatio">Ratio of source files which must include a header for it to be included in the pch</param>
 		/// <param name="Log">Writer for log output</param>
-		public static void GeneratePCHs(BuildTarget Target, Dictionary<SourceFile, CompileEnvironment> SourceFileToCompileEnvironment, IEnumerable<DirectoryReference> IncludePaths, IEnumerable<DirectoryReference> SystemIncludePaths, Dictionary<FileReference, string> OutputFileContents, float MinRatio, TextWriter Log)
+		public static void GeneratePCHs(BuildTarget Target, Dictionary<SourceFile, CompileEnvironment> SourceFileToCompileEnvironment, IEnumerable<DirectoryReference> IncludePaths, IEnumerable<DirectoryReference> SystemIncludePaths, Dictionary<FileReference, string> OutputFileContents, float MinRatio, LineBasedTextWriter Log)
 		{
 			Log.WriteLine("Optimizing precompiled headers...");
 
@@ -135,7 +135,7 @@ namespace IncludeTool.Reports
 		/// <param name="SourceFileToCompileEnvironment">Files being compiled</param>
 		/// <param name="PchToIncludeFileCount">Mapping of PCH to included files</param>
 		/// <param name="Log">Writer for log messages</param>
-		static void FindPchInfo(BuildTarget Target, Dictionary<SourceFile, CompileEnvironment> SourceFileToCompileEnvironment, Dictionary<FileReference, PchInfo> FileToPchInfo, TextWriter Log)
+		static void FindPchInfo(BuildTarget Target, Dictionary<SourceFile, CompileEnvironment> SourceFileToCompileEnvironment, Dictionary<FileReference, PchInfo> FileToPchInfo, LineBasedTextWriter Log)
 		{
 			// Create a map of module to the shared PCH it uses
 			Dictionary<BuildModule, FileReference> ModuleToPch = new Dictionary<BuildModule, FileReference>();
diff --git a/Engine/Source/Programs/IncludeTool/IncludeTool/Reports/SymbolReport.cs b/Engine/Source/Programs/IncludeTool/IncludeTool/Reports/SymbolReport.cs
--- a/Engine/Source/Programs/IncludeTool/IncludeTool/Reports/SymbolReport.cs
+++ b/Engine/Source/Programs/IncludeTool/IncludeTool/Reports/SymbolReport.cs
@@ -21,7 +21,7 @@ namespace IncludeTool
 		/// <param name="Files">The files to include in the report</param>
 		/// <param name="ReportFileLocation">Output file for the report</param>
 		/// <param name="Log">Writer for log output</param>
-		public static void Generate(FileReference ReportFileLocation, DirectoryReference InputDir, HashSet<SourceFile> PreprocessedFiles, TextWriter Log)
+		public static void Generate(FileReference ReportFileLocation, DirectoryReference InputDir, HashSet<SourceFile> PreprocessedFiles, LineBasedTextWriter Log)
 		{
 			Log.WriteLine("Writing {0}...", ReportFileLocation.FullName);
 
diff --git a/Engine/Source/Programs/IncludeTool/IncludeTool/SequenceProbe.cs b/Engine/Source/Programs/IncludeTool/IncludeTool/SequenceProbe.cs
--- a/Engine/Source/Programs/IncludeTool/IncludeTool/SequenceProbe.cs
+++ b/Engine/Source/Programs/IncludeTool/IncludeTool/SequenceProbe.cs
@@ -1,4 +1,4 @@
-﻿// Copyright 1998-2019 Epic Games, Inc. All Rights Reserved.
+// Copyright 1998-2019 Epic Games, Inc. All Rights Reserved.
 
 using System;
 using System.Collections.Generic;
@@ -441,7 +441,7 @@ namespace IncludeTool
 		/// Update the list of known dependencies to include dependencies of our already known dependencies
 		/// </summary>
 		/// <returns>True if the dependencies were updated, false if a dependency was encountered which is not available to this fragment</returns>
-		void UpdateDependencies(TextWriter Log)
+		void UpdateDependencies(LineBasedTextWriter Log)
 		{
 			// Update the task to account for any new dependencies that may have been found. Loop back through all the fragments so we include
 			// dependencies we find along the way.
@@ -493,7 +493,7 @@ namespace IncludeTool
 		/// <summary>
 		/// Start this worker
 		/// </summary>
-		public void Start(TextWriter Log)
+		public void Start(LineBasedTextWriter Log)
 		{
 			// Make sure there's not already a managed task in process
 			if(ActiveInstance != null)
@@ -534,7 +534,7 @@ namespace IncludeTool
 		/// <param name="StateFile">Path to the state file</param>
 		/// <param name="Writer">Writer for any messages</param>
 		/// <returns>Zero on success</returns>
-		static int Verify(string StateFile, TextWriter Writer)
+		static int Verify(string StateFile, LineBasedTextWriter Writer)
 		{
 			SequenceWorker Task = SequenceWorker.Deserialize(StateFile);
 			Task.Verify(Writer);
@@ -548,7 +548,7 @@ namespace IncludeTool
 		/// <param name="StateFile">Path to the state file</param>
 		/// <param name="Writer">Writer for any messages</param>
 		/// <returns>Zero on success</returns>
-		static int FindNextDependency(string StateFile, TextWriter Writer)
+		static int FindNextDependency(string StateFile, LineBasedTextWriter Writer)
 		{
 			SequenceWorker Task = SequenceWorker.Deserialize(StateFile);
 			Task.FindNextDependency(Writer);
@@ -561,7 +561,7 @@ namespace IncludeTool
 		/// </summary>
 		/// <param name="Writer">Writer for log output</param>
 		/// <returns>Return code from the thread</returns>
-		public SequenceProbeResult Join(TextWriter Writer)
+		public SequenceProbeResult Join(LineBasedTextWriter Writer)
 		{
 			// Finish the task instance
 			BufferedTextWriter BufferedWriter = new BufferedTextWriter();
diff --git a/Engine/Source/Programs/IncludeTool/IncludeTool/SequenceWorker.cs b/Engine/Source/Programs/IncludeTool/IncludeTool/SequenceWorker.cs
--- a/Engine/Source/Programs/IncludeTool/IncludeTool/SequenceWorker.cs
+++ b/Engine/Source/Programs/IncludeTool/IncludeTool/SequenceWorker.cs
@@ -1,4 +1,4 @@
-﻿// Copyright 1998-2019 Epic Games, Inc. All Rights Reserved.
+// Copyright 1998-2019 Epic Games, Inc. All Rights Reserved.
 
 using System;
 using System.Collections;
@@ -18,7 +18,7 @@ namespace IncludeTool
 	/// <summary>
 	/// Encapsulates the state of a sequence of fragments being optimized
 	/// </summary>
-	public class SequenceWorker
+	class SequenceWorker
 	{
 		/// <summary>
 		/// Path to the temporary file to write out permutations of the source file to try compiling
@@ -286,7 +286,7 @@ namespace IncludeTool
 		/// Compiles the permutation in its current state
 		/// </summary>
 		/// <param name="Writer">Writer for diagnostic messages</param>
-		public void Verify(TextWriter Writer)
+		public void Verify(LineBasedTextWriter Writer)
 		{
 			bResult = CompilePermutation(new List<int> { RemainingFragmentCount }, 0, 0, 0);
 		}
@@ -295,7 +295,7 @@ namespace IncludeTool
 		/// Find the next dependency
 		/// </summary>
 		/// <param name="Writer">Writer for diagnostic messages</param>
-		public void FindNextDependency(TextWriter Writer)
+		public void FindNextDependency(LineBasedTextWriter Writer)
 		{
 			// Update the required count for all the dependencies that are explicitly listed
 			while(KnownDependencies.Contains(RemainingFragmentCount - 1))
diff --git a/Engine/Source/Programs/IncludeTool/IncludeTool/SourceFile.cs b/Engine/Source/Programs/IncludeTool/IncludeTool/SourceFile.cs
--- a/Engine/Source/Programs/IncludeTool/IncludeTool/SourceFile.cs
+++ b/Engine/Source/Programs/IncludeTool/IncludeTool/SourceFile.cs
@@ -1,4 +1,4 @@
-﻿// Copyright 1998-2019 Epic Games, Inc. All Rights Reserved.
+// Copyright 1998-2019 Epic Games, Inc. All Rights Reserved.
 
 using IncludeTool.Support;
 using System;
@@ -378,7 +378,7 @@ namespace IncludeTool
 		/// <param name="IncludePaths">Base directories for relative include paths</param>
 		/// <param name="SystemIncludePaths">Base directories for system include paths</param>
 		/// <param name="Writer">Writer for the output text</param>
-		public void Write(IEnumerable<DirectoryReference> IncludePaths, IEnumerable<DirectoryReference> SystemIncludePaths, TextWriter Writer, bool bRemoveForwardDeclarations, TextWriter Log)
+		public void Write(IEnumerable<DirectoryReference> IncludePaths, IEnumerable<DirectoryReference> SystemIncludePaths, TextWriter Writer, bool bRemoveForwardDeclarations, LineBasedTextWriter Log)
 		{
 			// Write the file header
 			TextLocation LastLocation = Text.Start;
@@ -687,7 +687,7 @@ namespace IncludeTool
 		/// <param name="IncludePaths">Directories to base relative include paths from</param>
 		/// <param name="SystemIncludePaths">Directories to base system include paths from</param>
 		/// <returns>Formatted include path, with surrounding quotes</returns>
-		public static string FormatInclude(DirectoryReference FromDirectory, FileReference IncludeFile, IEnumerable<DirectoryReference> IncludePaths, IEnumerable<DirectoryReference> SystemIncludePaths, TextWriter Log)
+		public static string FormatInclude(DirectoryReference FromDirectory, FileReference IncludeFile, IEnumerable<DirectoryReference> IncludePaths, IEnumerable<DirectoryReference> SystemIncludePaths, LineBasedTextWriter Log)
 		{
 			string IncludeText;
 			if(!TryFormatInclude(FromDirectory, IncludeFile, IncludePaths, SystemIncludePaths, out IncludeText))
diff --git a/Engine/Source/Programs/IncludeTool/IncludeTool/Support/BufferedTextWriter.cs b/Engine/Source/Programs/IncludeTool/IncludeTool/Support/BufferedTextWriter.cs
--- a/Engine/Source/Programs/IncludeTool/IncludeTool/Support/BufferedTextWriter.cs
+++ b/Engine/Source/Programs/IncludeTool/IncludeTool/Support/BufferedTextWriter.cs
@@ -1,4 +1,4 @@
-﻿// Copyright 1998-2019 Epic Games, Inc. All Rights Reserved.
+// Copyright 1998-2019 Epic Games, Inc. All Rights Reserved.
 
 using System;
 using System.Collections.Generic;
@@ -12,18 +12,13 @@ namespace IncludeTool.Support
 	/// <summary>
 	/// Implementation of a TextWriter which buffers characters written to it, and can be copied to another TextWriter at a later time.
 	/// </summary>
-	class BufferedTextWriter : TextWriter
+	class BufferedTextWriter : LineBasedTextWriter
 	{
 		/// <summary>
 		/// List of buffered lines
 		/// </summary>
 		List<string> Lines = new List<string>();
 
-		/// <summary>
-		/// Buffered characters for the current line
-		/// </summary>
-		StringBuilder CurrentLine = new StringBuilder();
-
 		/// <summary>
 		/// Default constructor
 		/// </summary>
@@ -46,48 +41,27 @@ namespace IncludeTool.Support
 		public void Clear()
 		{
 			Lines.Clear();
-			CurrentLine.Clear();
 		}
 
 		/// <summary>
 		/// Copies the current contents of the BufferedTextWriter to another TextWriter
 		/// </summary>
 		/// <param name="Other">The TextWriter to copy the buffered text to</param>
-		public void CopyTo(TextWriter Other, string Prefix)
+		public void CopyTo(LineBasedTextWriter Other, string Prefix)
 		{
 			for (int Idx = 0; Idx < Lines.Count; Idx++)
 			{
 				Other.WriteLine("{0}{1}", Prefix, Lines[Idx]);
 			}
-			if(CurrentLine.Length > 0)
-			{
-				Other.Write("{0}{1}", Prefix, CurrentLine.ToString());
-			}
-		}
-
-		/// <summary>
-		/// Returns the encoding for this TextWriter
-		/// </summary>
-		public override Encoding Encoding
-		{
-			get { return Encoding.Unicode; }
 		}
 
 		/// <summary>
 		/// Write a single character to the buffer
 		/// </summary>
 		/// <param name="Character">Character to write</param>
-		public override void Write(char Character)
+		public override void WriteLine(string Line)
 		{
-			if (Character != '\r' && Character != '\n')
-			{
-				CurrentLine.Append(Character);
-			}
-			if (Character == '\n')
-			{
-				Lines.Add(CurrentLine.ToString());
-				CurrentLine.Clear();
-			}
+			Lines.Add(Line);
 		}
 
 		/// <summary>
@@ -96,7 +70,7 @@ namespace IncludeTool.Support
 		/// <returns>String containing the current buffered text</returns>
 		public override string ToString()
 		{
-			return String.Join("\n", Lines) + "\n" + CurrentLine.ToString();
+			return String.Join("\n", Lines);
 		}
 	}
 }
diff --git a/Engine/Source/Programs/IncludeTool/IncludeTool/Support/LineBasedTextWriter.cs b/Engine/Source/Programs/IncludeTool/IncludeTool/Support/LineBasedTextWriter.cs
new file mode 100644
--- /dev/null
+++ b/Engine/Source/Programs/IncludeTool/IncludeTool/Support/LineBasedTextWriter.cs
@@ -0,0 +1,26 @@
+using System;
+using System.Collections.Generic;
+using System.Linq;
+using System.Text;
+using System.Threading.Tasks;
+
+namespace IncludeTool.Support
+{
+	/// <summary>
+	/// Implements similar functionality to TextWriter, but outputs an entire line at a time. Used as a base class for logging.
+	/// </summary>
+	abstract class LineBasedTextWriter
+	{
+		public abstract void WriteLine(string Text);
+
+		public void WriteLine()
+		{
+			WriteLine("");
+		}
+
+		public void WriteLine(string Format, params object[] Args)
+		{
+			WriteLine(String.Format(Format, Args));
+		}
+	}
+}
diff --git a/Engine/Source/Programs/IncludeTool/IncludeTool/Support/LineBasedTextWriterWrapper.cs b/Engine/Source/Programs/IncludeTool/IncludeTool/Support/LineBasedTextWriterWrapper.cs
new file mode 100644
--- /dev/null
+++ b/Engine/Source/Programs/IncludeTool/IncludeTool/Support/LineBasedTextWriterWrapper.cs
@@ -0,0 +1,27 @@
+using System;
+using System.Collections.Generic;
+using System.IO;
+using System.Linq;
+using System.Text;
+using System.Threading.Tasks;
+
+namespace IncludeTool.Support
+{
+	/// <summary>
+	/// Implementation of LineBasedTextWriter which forwards data to a TextWriter instance.
+	/// </summary>
+	class LineBasedTextWriterWrapper : LineBasedTextWriter
+	{
+		TextWriter Inner;
+
+		public LineBasedTextWriterWrapper(TextWriter Inner)
+		{
+			this.Inner = Inner;
+		}
+
+		public override void WriteLine(string Line)
+		{
+			Inner.WriteLine(Line);
+		}
+	}
+}
diff --git a/Engine/Source/Programs/IncludeTool/IncludeTool/Support/LogWriter.cs b/Engine/Source/Programs/IncludeTool/IncludeTool/Support/LogWriter.cs
--- a/Engine/Source/Programs/IncludeTool/IncludeTool/Support/LogWriter.cs
+++ b/Engine/Source/Programs/IncludeTool/IncludeTool/Support/LogWriter.cs
@@ -1,4 +1,4 @@
-﻿// Copyright 1998-2019 Epic Games, Inc. All Rights Reserved.
+// Copyright 1998-2019 Epic Games, Inc. All Rights Reserved.
 
 using System;
 using System.Collections.Generic;
@@ -12,7 +12,7 @@ namespace IncludeTool.Support
 	/// <summary>
 	/// Provides for writing log messages to a log file and the display, as well as setting status messages
 	/// </summary>
-	class LogWriter : TextWriter
+	class LogWriter : LineBasedTextWriter, IDisposable
 	{
 		/// <summary>
 		/// Output stream for log text
@@ -32,7 +32,7 @@ namespace IncludeTool.Support
 		/// <summary>
 		/// Dispose of the current object, and its resources
 		/// </summary>
-		protected override void Dispose(bool Disposing)
+		public void Dispose()
 		{
 			if(FileWriter != null)
 			{
@@ -41,22 +41,14 @@ namespace IncludeTool.Support
 			}
 		}
 
-		/// <summary>
-		/// Returns the encoding of this log writer
-		/// </summary>
-		public override Encoding Encoding
-		{
-			get { return Encoding.Unicode; }
-		}
-
 		/// <summary>
 		/// Writes a character to the output stream
 		/// </summary>
 		/// <param name="Value">Character to write</param>
-		public override void Write(char Value)
+		public override void WriteLine(string Text)
 		{
-			FileWriter.Write(Value);
-			Console.Write(Value);
+			FileWriter.WriteLine(Text);
+			Console.WriteLine(Text);
 		}
 	}
 }
diff --git a/Engine/Source/Programs/IncludeTool/IncludeTool/Support/ManagedTask.cs b/Engine/Source/Programs/IncludeTool/IncludeTool/Support/ManagedTask.cs
--- a/Engine/Source/Programs/IncludeTool/IncludeTool/Support/ManagedTask.cs
+++ b/Engine/Source/Programs/IncludeTool/IncludeTool/Support/ManagedTask.cs
@@ -1,4 +1,4 @@
-﻿// Copyright 1998-2019 Epic Games, Inc. All Rights Reserved.
+// Copyright 1998-2019 Epic Games, Inc. All Rights Reserved.
 
 using Microsoft.Win32.SafeHandles;
 using System;
@@ -57,14 +57,14 @@ namespace IncludeTool.Support
 		/// </summary>
 		/// <param name="FinishedEvent">Event to set once the task is comlete</param>
 		/// <param name="Writer">Output writer for log events</param>
-		protected abstract void InternalStart(ManualResetEvent FinishedEvent, TextWriter Writer);
+		protected abstract void InternalStart(ManualResetEvent FinishedEvent, LineBasedTextWriter Writer);
 
 		/// <summary>
 		/// Waits for the task to complete and returns
 		/// </summary>
 		/// <param name="LogOutput">Receives the output text from the task</param>
 		/// <returns>The task's exit code</returns>
-		public int Join(TextWriter Writer)
+		public int Join(LineBasedTextWriter Writer)
 		{
 			FinishedEvent.WaitOne();
 
@@ -155,7 +155,7 @@ namespace IncludeTool.Support
 		/// <summary>
 		/// The action to execute
 		/// </summary>
-		Func<TextWriter, int> Action;
+		Func<LineBasedTextWriter, int> Action;
 
 		/// <summary>
 		/// Stores the result of executing the action
@@ -166,7 +166,7 @@ namespace IncludeTool.Support
 		/// Construct a managed task thread
 		/// </summary>
 		/// <param name="Action">The action to execute</param>
-		public ManagedTaskThread(Func<TextWriter, int> Action)
+		public ManagedTaskThread(Func<LineBasedTextWriter, int> Action)
 		{
 			this.Action = Action;
 		}
@@ -176,7 +176,7 @@ namespace IncludeTool.Support
 		/// </summary>
 		/// <param name="Event">Event to set when the action completes</param>
 		/// <param name="Writer">Writer for log output</param>
-		protected override void InternalStart(ManualResetEvent Event, TextWriter Writer)
+		protected override void InternalStart(ManualResetEvent Event, LineBasedTextWriter Writer)
 		{
 			Worker = new Thread(x => ExecuteAction(Event, Writer));
 			Worker.Start();
@@ -208,7 +208,7 @@ namespace IncludeTool.Support
 		/// </summary>
 		/// <param name="Event">Event to set when the action completes</param>
 		/// <param name="Writer">Writer for log output</param>
-		void ExecuteAction(ManualResetEvent Event, TextWriter Writer)
+		void ExecuteAction(ManualResetEvent Event, LineBasedTextWriter Writer)
 		{
 			Result = Action(Writer);
 			Event.Set();
@@ -251,7 +251,7 @@ namespace IncludeTool.Support
 		/// </summary>
 		/// <param name="Event">Event to set when the action completes</param>
 		/// <param name="Writer">Writer for log output</param>
-		protected override void InternalStart(ManualResetEvent Event, TextWriter Writer)
+		protected override void InternalStart(ManualResetEvent Event, LineBasedTextWriter Writer)
 		{
 			DataReceivedEventHandler OutputHandler = (Sender, Args) => { if(!String.IsNullOrEmpty(Args.Data)) { Writer.WriteLine(Args.Data); } };
 
diff --git a/Engine/Source/Programs/IncludeTool/IncludeTool/Support/Utility.cs b/Engine/Source/Programs/IncludeTool/IncludeTool/Support/Utility.cs
--- a/Engine/Source/Programs/IncludeTool/IncludeTool/Support/Utility.cs
+++ b/Engine/Source/Programs/IncludeTool/IncludeTool/Support/Utility.cs
@@ -1,4 +1,4 @@
-﻿// Copyright 1998-2019 Epic Games, Inc. All Rights Reserved.
+// Copyright 1998-2019 Epic Games, Inc. All Rights Reserved.
 
 using System;
 using System.Collections.Generic;
@@ -95,7 +95,7 @@ namespace IncludeTool.Support
 		/// <param name="Arguments">Arguments for the program</param>
 		/// <param name="WorkingDir">The working directory to run in</param>
 		/// <param name="Log">Writer to receive output messages</param>
-		public static int Run(FileReference ExecutableFile, string Arguments, DirectoryReference WorkingDir, TextWriter Log)
+		public static int Run(FileReference ExecutableFile, string Arguments, DirectoryReference WorkingDir, LineBasedTextWriter Log)
 		{
 			using (Process NewProcess = new Process())
 			{
@@ -224,7 +224,7 @@ namespace IncludeTool.Support
 			/// <param name="TotalCount">The total number of iterations</param>
 			/// <param name="Message">The message to output</param>
 			/// <param name="Log">Writer for output messages</param>
-			public void Increment(int TotalCount, string Message, TextWriter Log)
+			public void Increment(int TotalCount, string Message, LineBasedTextWriter Log)
 			{
 				lock(this)
 				{
@@ -247,7 +247,7 @@ namespace IncludeTool.Support
 			/// <param name="TotalCount">The total number of iterations</param>
 			/// <param name="Message">The message to output</param>
 			/// <param name="Log">Writer for output messages</param>
-			public void Complete(int TotalCount, string Message, TextWriter Log)
+			public void Complete(int TotalCount, string Message, LineBasedTextWriter Log)
 			{
 				Log.WriteLine("[{0}] {1} ({2}/{3}; {4}%)", Timer.Elapsed.ToString(@"hh\:mm\:ss"), Message, TotalCount, TotalCount, 100);
 			}
@@ -261,7 +261,7 @@ namespace IncludeTool.Support
 		/// <param name="EndValue">The upper bound for the for loop, exclusive</param>
 		/// <param name="Action">The action to execute for each iteration</param>
 		/// <param name="Log">Log for output messages</param>
-		public static void ParallelForWithStatus(string Message, int BeginValue, int EndValue, Action<int> Action, TextWriter Log)
+		public static void ParallelForWithStatus(string Message, int BeginValue, int EndValue, Action<int> Action, LineBasedTextWriter Log)
 		{
 			ParallelForWithStatus(Message, BeginValue, EndValue, new ParallelOptions(), Action, Log);
 		}
@@ -274,7 +274,7 @@ namespace IncludeTool.Support
 		/// <param name="EndValue">The upper bound for the for loop, exclusive</param>
 		/// <param name="Action">The action to execute for each iteration</param>
 		/// <param name="Log">Log for output messages</param>
-		public static void ParallelForWithStatus(string Message, int BeginValue, int EndValue, ParallelOptions Options, Action<int> Action, TextWriter Log)
+		public static void ParallelForWithStatus(string Message, int BeginValue, int EndValue, ParallelOptions Options, Action<int> Action, LineBasedTextWriter Log)
 		{
 			Log.WriteLine(Message);
 
diff --git a/Engine/Source/Programs/IncludeTool/IncludeTool/SymbolTable.cs b/Engine/Source/Programs/IncludeTool/IncludeTool/SymbolTable.cs
--- a/Engine/Source/Programs/IncludeTool/IncludeTool/SymbolTable.cs
+++ b/Engine/Source/Programs/IncludeTool/IncludeTool/SymbolTable.cs
@@ -166,7 +166,7 @@ namespace IncludeTool
 		/// Prints all the symbols with conflicting definitions
 		/// </summary>
 		/// <param name="Log">Writer for log messagsThe source fragment to parse</param>
-		public void PrintConflicts(TextWriter Log)
+		public void PrintConflicts(LineBasedTextWriter Log)
 		{
 			foreach(string SymbolName in Lookup.Keys)
 			{
@@ -701,7 +701,7 @@ namespace IncludeTool
 		/// <param name="HeaderFile">The file to read from</param>
 		/// <param name="SymbolToHeader">Map of symbol to the file containing it</param>
 		/// <param name="Log">Writer for warnings and errors</param>
-		public bool ReadForwardDeclarations(SourceFile HeaderFile, Dictionary<Symbol, SourceFile> SymbolToHeader, TextWriter Log)
+		public bool ReadForwardDeclarations(SourceFile HeaderFile, Dictionary<Symbol, SourceFile> SymbolToHeader, LineBasedTextWriter Log)
 		{
 			foreach (PreprocessorMarkup Markup in HeaderFile.Markup.Where(x => x.Type == PreprocessorMarkupType.Text))
 			{
@@ -729,7 +729,7 @@ namespace IncludeTool
 		/// <param name="SymbolToHeader">Map of symbol to the file containing it</param>
 		/// <param name="Log">Writer for warnings and errors</param>
 		/// <returns>True if a forward declaration was read, false otherwise</returns>
-		bool ReadClassOrStructForwardDeclaration(TokenReader OriginalReader, SourceFile HeaderFile, Dictionary<Symbol, SourceFile> SymbolToHeader, TextWriter Log)
+		bool ReadClassOrStructForwardDeclaration(TokenReader OriginalReader, SourceFile HeaderFile, Dictionary<Symbol, SourceFile> SymbolToHeader, LineBasedTextWriter Log)
 		{
 			SymbolType Type;
 			if(OriginalReader.Current.Text == "class")
@@ -770,7 +770,7 @@ namespace IncludeTool
 		/// <param name="SymbolToHeader">Map of symbol to the file containing it</param>
 		/// <param name="Log">Writer for warnings and errors</param>
 		/// <returns>True if a forward declaration was read, false otherwise</returns>
-		bool ReadEnumClassForwardDeclaration(TokenReader OriginalReader, SourceFile HeaderFile, Dictionary<Symbol, SourceFile> SymbolToHeader, TextWriter Log)
+		bool ReadEnumClassForwardDeclaration(TokenReader OriginalReader, SourceFile HeaderFile, Dictionary<Symbol, SourceFile> SymbolToHeader, LineBasedTextWriter Log)
 		{
 			if(OriginalReader.Current.Text != "enum")
 			{
@@ -813,7 +813,7 @@ namespace IncludeTool
 		/// <param name="SymbolToHeader">Map of symbol to the file containing it</param>
 		/// <param name="Log">Writer for warnings and errors</param>
 		/// <returns>True if a forward declaration was read, false otherwise</returns>
-		bool ReadTemplateClassOrStructForwardDeclaration(TokenReader OriginalReader, SourceFile HeaderFile, Dictionary<Symbol, SourceFile> SymbolToHeader, TextWriter Log)
+		bool ReadTemplateClassOrStructForwardDeclaration(TokenReader OriginalReader, SourceFile HeaderFile, Dictionary<Symbol, SourceFile> SymbolToHeader, LineBasedTextWriter Log)
 		{
 			if(OriginalReader.Current.Text != "template")
 			{
@@ -945,7 +945,7 @@ namespace IncludeTool
 		/// <param name="Type"></param>
 		/// <param name="SymbolToHeader"></param>
 		/// <param name="Log"></param>
-		void AddForwardDeclaration(SourceFile HeaderFile, string SymbolName, SymbolType Type, Dictionary<Symbol, SourceFile> SymbolToHeader, TextWriter Log)
+		void AddForwardDeclaration(SourceFile HeaderFile, string SymbolName, SymbolType Type, Dictionary<Symbol, SourceFile> SymbolToHeader, LineBasedTextWriter Log)
 		{
 			foreach (Symbol Symbol in Lookup.WithKey(SymbolName))
 			{
