Add missing copyright notices.

diff --git a/Engine/Source/Programs/IncludeTool/IncludeTool/Support/LineBasedTextWriter.cs b/Engine/Source/Programs/IncludeTool/IncludeTool/Support/LineBasedTextWriter.cs
--- a/Engine/Source/Programs/IncludeTool/IncludeTool/Support/LineBasedTextWriter.cs
+++ b/Engine/Source/Programs/IncludeTool/IncludeTool/Support/LineBasedTextWriter.cs
@@ -1,3 +1,5 @@
+// Copyright 1998-2019 Epic Games, Inc. All Rights Reserved.
+
 using System;
 using System.Collections.Generic;
 using System.Linq;
diff --git a/Engine/Source/Programs/IncludeTool/IncludeTool/Support/LineBasedTextWriterWrapper.cs b/Engine/Source/Programs/IncludeTool/IncludeTool/Support/LineBasedTextWriterWrapper.cs
--- a/Engine/Source/Programs/IncludeTool/IncludeTool/Support/LineBasedTextWriterWrapper.cs
+++ b/Engine/Source/Programs/IncludeTool/IncludeTool/Support/LineBasedTextWriterWrapper.cs
@@ -1,3 +1,5 @@
+// Copyright 1998-2019 Epic Games, Inc. All Rights Reserved.
+
 using System;
 using System.Collections.Generic;
 using System.IO;
