Improve IncludeTool diagnostic messages.

#rnx

diff --git a/Engine/Source/Programs/IncludeTool/IncludeTool/Program.cs b/Engine/Source/Programs/IncludeTool/IncludeTool/Program.cs
--- a/Engine/Source/Programs/IncludeTool/IncludeTool/Program.cs
+++ b/Engine/Source/Programs/IncludeTool/IncludeTool/Program.cs
@@ -1049,13 +1049,13 @@ namespace IncludeTool
 							{
 								lock(Log)
 								{
-									Log.WriteLine("{0}({1}): warning: inconsistent preprocessor state.", File.Location.FullName, File.Markup[Idx].Location.LineIdx + 1);
-									Log.WriteLine("    First include stack ({0}):", CurrentActiveMarkup.Flags[Idx]? "active" : "inactive");
+									Log.WriteLine("{0}({1}): warning: Inconsistent preprocessor state. Every block of code in a header file should consistently be active or inactive regardless of where it is included from. Inconsistencies prevent IWYU optimization, and may be indicative of programming errors.", File.Location.FullName, File.Markup[Idx].Location.LineIdx + 1);
+									Log.WriteLine("    First include stack ({0}):", CurrentActiveMarkup.Flags[Idx]? "code is active" : "code is inactive");
 									foreach(FileReference IncludeFile in CurrentActiveMarkup.IncludeStack.Reverse())
 									{
 										Log.WriteLine("        {0}", IncludeFile);
 									}
-									Log.WriteLine("    Second include stack ({0}):", ActiveMarkup.Flags[Idx]? "active" : "inactive");
+									Log.WriteLine("    Second include stack ({0}):", ActiveMarkup.Flags[Idx]? "code is active" : "code is inactive");
 									foreach(FileReference IncludeFile in ActiveMarkup.IncludeStack.Reverse())
 									{
 										Log.WriteLine("        {0}", IncludeFile);
diff --git a/Engine/Source/Programs/IncludeTool/IncludeTool/Rules.cs b/Engine/Source/Programs/IncludeTool/IncludeTool/Rules.cs
--- a/Engine/Source/Programs/IncludeTool/IncludeTool/Rules.cs
+++ b/Engine/Source/Programs/IncludeTool/IncludeTool/Rules.cs
@@ -461,6 +461,7 @@ namespace IncludeTool
 			"/Engine/Source/Runtime/Slate/Public/Framework/Text/ShapedTextCacheFwd.h", // Typedef isn't a forward declaration
 			"/Engine/Source/Runtime/MovieScene/Public/MovieSceneFwd.h",
 			"/Engine/Source/Runtime/Core/Public/Internationalization/StringTableCoreFwd.h", // Typedef isn't a forward declaration
+			"/Engine/Source/Runtime/Experimental/Chaos/Public/Chaos/ParticleHandleFwd.h", // invalid forward declaration - 'namespace Chaos'
 		};
 
 		/// <summary>
