Output a warning rather than a hard error for include paths that have multiple slashes, and change IncludeTool to use faster string comparisons.

#rnx

diff --git a/Engine/Source/Programs/IncludeTool/IncludeTool/BuildTarget.cs b/Engine/Source/Programs/IncludeTool/IncludeTool/BuildTarget.cs
--- a/Engine/Source/Programs/IncludeTool/IncludeTool/BuildTarget.cs
+++ b/Engine/Source/Programs/IncludeTool/IncludeTool/BuildTarget.cs
@@ -1,4 +1,4 @@
-﻿// Copyright 1998-2019 Epic Games, Inc. All Rights Reserved.
+// Copyright 1998-2019 Epic Games, Inc. All Rights Reserved.
 
 using IncludeTool.Support;
 using System;
@@ -107,7 +107,7 @@ namespace IncludeTool
 			}
 
 			JsonObject Modules = Object.GetObjectField("Modules");
-			Target.NameToModule = Modules.KeyNames.ToDictionary(x => x, x => new BuildModule { Name = x }, StringComparer.InvariantCultureIgnoreCase);
+			Target.NameToModule = Modules.KeyNames.ToDictionary(x => x, x => new BuildModule { Name = x }, StringComparer.OrdinalIgnoreCase);
 			Target.Modules.AddRange(Target.NameToModule.Values);
 
 			foreach (BuildModule Module in Target.Modules)
diff --git a/Engine/Source/Programs/IncludeTool/IncludeTool/Preprocessor.cs b/Engine/Source/Programs/IncludeTool/IncludeTool/Preprocessor.cs
--- a/Engine/Source/Programs/IncludeTool/IncludeTool/Preprocessor.cs
+++ b/Engine/Source/Programs/IncludeTool/IncludeTool/Preprocessor.cs
@@ -467,8 +467,8 @@ namespace IncludeTool
 			ExpandMacros(Tokens, ExpandedTokens, false, CurrentLine);
 
 			// Expand any macros in them and resolve it
-			string IncludeToken = Token.Format(ExpandedTokens);
-			if(IncludeToken.Length >= 2 && IncludeToken.StartsWith("\"") && IncludeToken.EndsWith("\""))
+			string IncludeToken = Token.Format(ExpandedTokens).Replace("//", "/");
+			if (IncludeToken.Length >= 2 && IncludeToken.StartsWith("\"") && IncludeToken.EndsWith("\""))
 			{
 				return ResolveQuotedInclude(IncludeToken.Substring(1, IncludeToken.Length - 2));
 			}
diff --git a/Engine/Source/Programs/IncludeTool/IncludeTool/Program.cs b/Engine/Source/Programs/IncludeTool/IncludeTool/Program.cs
--- a/Engine/Source/Programs/IncludeTool/IncludeTool/Program.cs
+++ b/Engine/Source/Programs/IncludeTool/IncludeTool/Program.cs
@@ -420,7 +420,7 @@ namespace IncludeTool
 				}
 
 				// Find a lookup of all the files by name
-				MultiValueDictionary<string, SourceFile> NameToFile = new MultiValueDictionary<string, SourceFile>(StringComparer.InvariantCultureIgnoreCase);
+				MultiValueDictionary<string, SourceFile> NameToFile = new MultiValueDictionary<string, SourceFile>(StringComparer.OrdinalIgnoreCase);
 				foreach(SourceFile PreprocessedFile in PreprocessedFiles)
 				{
 					NameToFile.Add(PreprocessedFile.Location.GetFileName(), PreprocessedFile);
@@ -583,7 +583,7 @@ namespace IncludeTool
 					{
 						if(!PreprocessedFile.HasHeaderGuard && (PreprocessedFile.Flags & (SourceFileFlags.TranslationUnit | SourceFileFlags.Inline | SourceFileFlags.External | SourceFileFlags.GeneratedHeader)) == 0 && PreprocessedFile.Counterpart == null && PreprocessedFile.Location.HasExtension(".h"))
 						{
-							if(!PreprocessedFile.Location.GetFileName().Equals("MonolithicHeaderBoilerplate.h", StringComparison.InvariantCultureIgnoreCase))
+							if(!PreprocessedFile.Location.GetFileName().Equals("MonolithicHeaderBoilerplate.h", StringComparison.OrdinalIgnoreCase))
 							{
 								Log.WriteLine("warning: missing header guard: {0}", PreprocessedFile.Location.FullName);
 							}
@@ -893,7 +893,7 @@ namespace IncludeTool
 		/// <returns>True if the given path is to a generated cpp file</returns>
 		static bool IsGeneratedCppFile(FileReference Location)
 		{
-			return Location.FullName.EndsWith(".cpp", StringComparison.InvariantCultureIgnoreCase) && Location.FullName.IndexOf(".generated.", StringComparison.InvariantCultureIgnoreCase) != -1;
+			return Location.FullName.EndsWith(".cpp", StringComparison.OrdinalIgnoreCase) && Location.FullName.IndexOf(".generated.", StringComparison.OrdinalIgnoreCase) != -1;
 		}
 
 		/// <summary>
@@ -986,24 +986,32 @@ namespace IncludeTool
 					PreprocessorMarkup Markup = File.Markup[Idx];
 					try
 					{
-						if(Markup.Type != PreprocessorMarkupType.Include)
+						if (Markup.Type != PreprocessorMarkupType.Include)
 						{
 							Preprocessor.ParseMarkup(Markup.Type, Markup.Tokens, Markup.Location.LineIdx);
 						}
-						else if(Preprocessor.IsBranchActive())
+						else
 						{
-							// Figure out which file is being included
-							PreprocessorFile TargetFile = Preprocessor.ResolveInclude(Markup.Tokens, Markup.Location.LineIdx);
-							if(!Markup.SetIncludedFile(TargetFile.WorkspaceFile.ReadSourceFile()))
+							if (Markup.Tokens.Count == 1 && (Markup.Tokens[0].Type == TokenType.StringLiteral || Markup.Tokens[0].Type == TokenType.SystemInclude) && Markup.Tokens[0].Text.Contains("//"))
+							{
+								Log.WriteLine("{0}({1}): warning: include path has multiple slashes ({2})", File.Location, Markup.Location.LineIdx + 1, Markup.Tokens[0].Text);
+							}
+
+							if (Preprocessor.IsBranchActive())
 							{
-								PreprocessorFile ExistingTargetFile = new PreprocessorFile(TargetFile.FileName, Workspace.GetFile(Markup.IncludedFile.Location));
-								Log.WriteLine("{0}({1}): warning: include text resolved to different locations", File.Location, Markup.Location.LineIdx + 1);
-								Log.WriteLine("    First: {0}", ExistingTargetFile.Location);
-								Log.WriteLine("    Now:   {0}", TargetFile.Location);
-								TargetFile = ExistingTargetFile;
+								// Figure out which file is being included
+								PreprocessorFile TargetFile = Preprocessor.ResolveInclude(Markup.Tokens, Markup.Location.LineIdx);
+								if (!Markup.SetIncludedFile(TargetFile.WorkspaceFile.ReadSourceFile()))
+								{
+									PreprocessorFile ExistingTargetFile = new PreprocessorFile(TargetFile.FileName, Workspace.GetFile(Markup.IncludedFile.Location));
+									Log.WriteLine("{0}({1}): warning: include text resolved to different locations", File.Location, Markup.Location.LineIdx + 1);
+									Log.WriteLine("    First: {0}", ExistingTargetFile.Location);
+									Log.WriteLine("    Now:   {0}", TargetFile.Location);
+									TargetFile = ExistingTargetFile;
+								}
+								PreprocessIncludedFile(Preprocessor, TargetFile, FileToActiveMarkup, Log);
+								ActiveMarkup.Flags[Idx] = true;
 							}
-							PreprocessIncludedFile(Preprocessor, TargetFile, FileToActiveMarkup, Log);
-							ActiveMarkup.Flags[Idx] = true;
 						}
 					}
 					catch(PreprocessorException Ex)
@@ -1096,7 +1104,7 @@ namespace IncludeTool
 		/// <param name="Files">Array of files to create fragments for</param>
 		static void CreateFragments(IEnumerable<SourceFile> Files)
 		{
-			HashSet<string> UniqueNames = new HashSet<string>(StringComparer.InvariantCultureIgnoreCase);
+			HashSet<string> UniqueNames = new HashSet<string>(StringComparer.OrdinalIgnoreCase);
 			foreach(SourceFile File in Files)
 			{
 				CreateFragmentsRecursively(File, UniqueNames);
@@ -1266,7 +1274,7 @@ namespace IncludeTool
 			Dictionary<SourceFragment, Lazy<SequenceProbe>> FragmentToLazyProbe = new Dictionary<SourceFragment, Lazy<SequenceProbe>>();
 
 			// Create all the unique probes
-			HashSet<string> UniqueProbeNames = new HashSet<string>(StringComparer.InvariantCultureIgnoreCase);
+			HashSet<string> UniqueProbeNames = new HashSet<string>(StringComparer.OrdinalIgnoreCase);
 			foreach(SourceFile SourceFile in SourceFileToCompileEnvironment.Keys)
 			{
 				// Flatten the file into a list of fragments
diff --git a/Engine/Source/Programs/IncludeTool/IncludeTool/Rules.cs b/Engine/Source/Programs/IncludeTool/IncludeTool/Rules.cs
--- a/Engine/Source/Programs/IncludeTool/IncludeTool/Rules.cs
+++ b/Engine/Source/Programs/IncludeTool/IncludeTool/Rules.cs
@@ -260,7 +260,7 @@ namespace IncludeTool
 		/// <summary>
 		/// List of context-specific files
 		/// </summary>
-		static readonly HashSet<string> InlineFileNames = new HashSet<string>(StringComparer.InvariantCultureIgnoreCase)
+		static readonly HashSet<string> InlineFileNames = new HashSet<string>(StringComparer.OrdinalIgnoreCase)
 		{
 			//"/Engine/Source/Runtime/Core/Public/Delegates/DelegateCombinations.h",
 			"/Engine/Source/Runtime/Core/Public/UObject/PendingVersions.h",
@@ -307,7 +307,7 @@ namespace IncludeTool
 		/// <summary>
 		/// List of files whose includes are pinned to the file they are included from
 		/// </summary>
-		static readonly HashSet<string> PinnedFileNames = new HashSet<string>(StringComparer.InvariantCultureIgnoreCase)
+		static readonly HashSet<string> PinnedFileNames = new HashSet<string>(StringComparer.OrdinalIgnoreCase)
 		{
 			"/Engine/Source/Runtime/Core/Public/GenericPlatform/GenericPlatform.h",
 			"/Engine/Source/Runtime/Core/Public/Math/UnrealMatrix.h",
@@ -400,7 +400,7 @@ namespace IncludeTool
 		/// <summary>
 		/// List of files whose includes are pinned to the file they are included from
 		/// </summary>
-		static readonly HashSet<string> NotStandaloneFileNames = new HashSet<string>(StringComparer.InvariantCultureIgnoreCase)
+		static readonly HashSet<string> NotStandaloneFileNames = new HashSet<string>(StringComparer.OrdinalIgnoreCase)
 		{
 			"/Engine/Source/Runtime/Core/Public/Delegates/DelegateCombinations.h",
 			"/Engine/Source/Runtime/Core/Public/Containers/LockFreeListImpl.h",
@@ -419,7 +419,7 @@ namespace IncludeTool
 		/// <summary>
 		/// List of files whose includes are pinned to the file they are included from
 		/// </summary>
-		static readonly HashSet<string> AggregateFileNames = new HashSet<string>(StringComparer.InvariantCultureIgnoreCase)
+		static readonly HashSet<string> AggregateFileNames = new HashSet<string>(StringComparer.OrdinalIgnoreCase)
 		{
 //			"/Engine/Source/Runtime/CoreUObject/Classes/Object.h",
 			"/Engine/Source/Runtime/Slate/Public/SlateBasics.h",
@@ -429,7 +429,7 @@ namespace IncludeTool
 		/// <summary>
 		/// List of files whose includes are pinned to the file they are included from
 		/// </summary>
-		static readonly HashSet<string> NonAggregateFileNames = new HashSet<string>(StringComparer.InvariantCultureIgnoreCase)
+		static readonly HashSet<string> NonAggregateFileNames = new HashSet<string>(StringComparer.OrdinalIgnoreCase)
 		{
 			"/Engine/Source/Runtime/RHI/Public/RHI.h",
 		};
@@ -437,7 +437,7 @@ namespace IncludeTool
 		/// <summary>
 		/// List of files which are ok to split into multiple fragments
 		/// </summary>
-		static readonly HashSet<string> AllowMultipleFragmentFileNames = new HashSet<string>(StringComparer.InvariantCultureIgnoreCase)
+		static readonly HashSet<string> AllowMultipleFragmentFileNames = new HashSet<string>(StringComparer.OrdinalIgnoreCase)
 		{
 			"/Engine/Source/Runtime/CoreUObject/Public/UObject/UObjectBaseUtility.h",
 			"/Engine/Source/Runtime/CoreUObject/Public/UObject/Class.h",
@@ -453,7 +453,7 @@ namespace IncludeTool
 		/// <summary>
 		/// Files which should not be treated as containing purely forward declarations, despite a "fwd.h" suffix
 		/// </summary>
-		static readonly HashSet<string> IgnoreFwdHeaders = new HashSet<string>(StringComparer.InvariantCultureIgnoreCase)
+		static readonly HashSet<string> IgnoreFwdHeaders = new HashSet<string>(StringComparer.OrdinalIgnoreCase)
 		{
 			"/Engine/Source/Runtime/Core/Public/Internationalization/TextNamespaceFwd.h",
 			"/Engine/Source/Editor/SceneOutliner/Public/SceneOutlinerFwd.h",
diff --git a/Engine/Source/Programs/IncludeTool/IncludeTool/SourceFile.cs b/Engine/Source/Programs/IncludeTool/IncludeTool/SourceFile.cs
--- a/Engine/Source/Programs/IncludeTool/IncludeTool/SourceFile.cs
+++ b/Engine/Source/Programs/IncludeTool/IncludeTool/SourceFile.cs
@@ -451,7 +451,7 @@ namespace IncludeTool
 					while(NewLastLocation.LineIdx < Text.Lines.Length)
 					{
 						string TrimLine = Text.Lines[NewLastLocation.LineIdx].Trim();
-						if(TrimLine.Length > 0 && !TrimLine.Equals("// Forward declarations", StringComparison.InvariantCultureIgnoreCase) && !TrimLine.Equals("// Forward declarations.", StringComparison.InvariantCultureIgnoreCase))
+						if(TrimLine.Length > 0 && !TrimLine.Equals("// Forward declarations", StringComparison.OrdinalIgnoreCase) && !TrimLine.Equals("// Forward declarations.", StringComparison.OrdinalIgnoreCase))
 						{
 							// Create a token reader for the current line
 							TokenReader Reader = new TokenReader(Text, new TextLocation(NewLastLocation.LineIdx, 0), new TextLocation(NewLastLocation.LineIdx, Text.Lines[NewLastLocation.LineIdx].Length));
@@ -540,7 +540,7 @@ namespace IncludeTool
 						{
 							IncludeText = OriginalIncludeText;
 						}
-						else if(OriginalIncludeText != null && (Flags & SourceFileFlags.TranslationUnit) == 0 && OriginalIncludeText.EndsWith(IncludeText.TrimStart('\"'), StringComparison.InvariantCultureIgnoreCase) && (OriginalIncludeText.StartsWith("\"Runtime/", StringComparison.InvariantCultureIgnoreCase) || OriginalIncludeText.StartsWith("\"Developer/", StringComparison.InvariantCultureIgnoreCase) || OriginalIncludeText.StartsWith("\"Editor/", StringComparison.InvariantCultureIgnoreCase)))
+						else if(OriginalIncludeText != null && (Flags & SourceFileFlags.TranslationUnit) == 0 && OriginalIncludeText.EndsWith(IncludeText.TrimStart('\"'), StringComparison.OrdinalIgnoreCase) && (OriginalIncludeText.StartsWith("\"Runtime/", StringComparison.InvariantCultureIgnoreCase) || OriginalIncludeText.StartsWith("\"Developer/", StringComparison.InvariantCultureIgnoreCase) || OriginalIncludeText.StartsWith("\"Editor/", StringComparison.InvariantCultureIgnoreCase)))
 						{
 							IncludeText = OriginalIncludeText;
 						}
@@ -743,26 +743,26 @@ namespace IncludeTool
 			}
 
 			// HACK: VsPerf.h is in the compiler environment, but the include path is added directly
-			if (IncludeFile.FullName.EndsWith("\\PerfSDK\\VSPerf.h", StringComparison.InvariantCultureIgnoreCase))
+			if (IncludeFile.FullName.EndsWith("\\PerfSDK\\VSPerf.h", StringComparison.OrdinalIgnoreCase))
 			{
 				IncludeText = "\"VSPerf.h\"";
 				return true;
 			}
 
 			// HACK: public Paper2D header in classes folder including private Paper2D header
-			if(IncludeFile.FullName.IndexOf("\\Paper2D\\Private\\", StringComparison.InvariantCultureIgnoreCase) != -1)
+			if(IncludeFile.FullName.IndexOf("\\Paper2D\\Private\\", StringComparison.OrdinalIgnoreCase) != -1)
 			{
 				IncludeText = "\"" + IncludeFile.GetFileName() + "\"";
 				return true;
 			}
-			if(IncludeFile.FullName.IndexOf("\\OnlineSubsystemUtils\\Private\\", StringComparison.InvariantCultureIgnoreCase) != -1)
+			if(IncludeFile.FullName.IndexOf("\\OnlineSubsystemUtils\\Private\\", StringComparison.OrdinalIgnoreCase) != -1)
 			{
 				IncludeText = "\"" + IncludeFile.GetFileName() + "\"";
 				return true;
 			}
 
 			// HACK: including private headers from public headers in the same module
-			int PrivateIdx = IncludeFile.FullName.IndexOf("\\Private\\", StringComparison.InvariantCultureIgnoreCase);
+			int PrivateIdx = IncludeFile.FullName.IndexOf("\\Private\\", StringComparison.OrdinalIgnoreCase);
 			if(PrivateIdx != -1)
 			{
 				DirectoryReference BaseDir = new DirectoryReference(IncludeFile.FullName.Substring(0, PrivateIdx));
diff --git a/Engine/Source/Programs/IncludeTool/IncludeTool/Support/CommandLine.cs b/Engine/Source/Programs/IncludeTool/IncludeTool/Support/CommandLine.cs
--- a/Engine/Source/Programs/IncludeTool/IncludeTool/Support/CommandLine.cs
+++ b/Engine/Source/Programs/IncludeTool/IncludeTool/Support/CommandLine.cs
@@ -1,4 +1,4 @@
-﻿// Copyright 1998-2019 Epic Games, Inc. All Rights Reserved.
+// Copyright 1998-2019 Epic Games, Inc. All Rights Reserved.
 
 using System;
 using System.Collections;
@@ -57,7 +57,7 @@ namespace IncludeTool.Support
 					if (Field.FieldType == typeof(bool))
 					{
 						string Flag = String.Format("-{0}", Name);
-						if(RemainingArgs.RemoveAll(x => x.Equals(Flag, StringComparison.InvariantCultureIgnoreCase)) > 0)
+						if(RemainingArgs.RemoveAll(x => x.Equals(Flag, StringComparison.OrdinalIgnoreCase)) > 0)
 						{
 							Field.SetValue(Options, true);
 						}
@@ -70,7 +70,7 @@ namespace IncludeTool.Support
 						List<string> Values = new List<string>();
 						for (int Idx = 0; Idx < RemainingArgs.Count; Idx++)
 						{
-							if (RemainingArgs[Idx].StartsWith(Prefix, StringComparison.InvariantCultureIgnoreCase))
+							if (RemainingArgs[Idx].StartsWith(Prefix, StringComparison.OrdinalIgnoreCase))
 							{
 								string Suffix = RemainingArgs[Idx].Substring(Prefix.Length);
 								if (OptionAttribute.Separator != 0)
diff --git a/Engine/Source/Programs/IncludeTool/IncludeTool/Support/JsonObject.cs b/Engine/Source/Programs/IncludeTool/IncludeTool/Support/JsonObject.cs
--- a/Engine/Source/Programs/IncludeTool/IncludeTool/Support/JsonObject.cs
+++ b/Engine/Source/Programs/IncludeTool/IncludeTool/Support/JsonObject.cs
@@ -1,4 +1,4 @@
-﻿// Copyright 1998-2019 Epic Games, Inc. All Rights Reserved.
+// Copyright 1998-2019 Epic Games, Inc. All Rights Reserved.
 
 using System;
 using System.Collections.Generic;
@@ -23,7 +23,7 @@ namespace IncludeTool.Support
 
 		public JsonObject(Dictionary<string, object> InRawObject)
 		{
-			RawObject = new Dictionary<string, object>(InRawObject, StringComparer.InvariantCultureIgnoreCase);
+			RawObject = new Dictionary<string, object>(InRawObject, StringComparer.OrdinalIgnoreCase);
 		}
 
 		public static JsonObject Read(string FileName)
diff --git a/Engine/Source/Programs/IncludeTool/IncludeTool/Workspace.cs b/Engine/Source/Programs/IncludeTool/IncludeTool/Workspace.cs
--- a/Engine/Source/Programs/IncludeTool/IncludeTool/Workspace.cs
+++ b/Engine/Source/Programs/IncludeTool/IncludeTool/Workspace.cs
@@ -1,4 +1,4 @@
-﻿// Copyright 1998-2019 Epic Games, Inc. All Rights Reserved.
+// Copyright 1998-2019 Epic Games, Inc. All Rights Reserved.
 
 using IncludeTool.Support;
 using System;
@@ -168,7 +168,7 @@ namespace IncludeTool
 			if(NameToFile == null)
 			{
 				// Add all the files in this directory
-				Dictionary<string, WorkspaceFile> NewNameToFile = new Dictionary<string, WorkspaceFile>(StringComparer.InvariantCultureIgnoreCase);
+				Dictionary<string, WorkspaceFile> NewNameToFile = new Dictionary<string, WorkspaceFile>(StringComparer.OrdinalIgnoreCase);
 				foreach (FileInfo FileInfo in Info.EnumerateFiles())
 				{
 					WorkspaceFile File = new WorkspaceFile(FileInfo.Name, this);
@@ -205,7 +205,7 @@ namespace IncludeTool
 		{
 			if(NameToDirectory == null)
 			{
-				Dictionary<string, WorkspaceDirectory> NewNameToDirectory = new Dictionary<string, WorkspaceDirectory>(StringComparer.InvariantCultureIgnoreCase);
+				Dictionary<string, WorkspaceDirectory> NewNameToDirectory = new Dictionary<string, WorkspaceDirectory>(StringComparer.OrdinalIgnoreCase);
 				foreach(DirectoryInfo DirectoryInfo in Info.EnumerateDirectories())
 				{
 					WorkspaceDirectory Directory = new WorkspaceDirectory(DirectoryInfo, this);
