Add another Chaos not fwd file IgnoreFwdHeaders list
#rnx

diff --git a/Engine/Source/Programs/IncludeTool/IncludeTool/Rules.cs b/Engine/Source/Programs/IncludeTool/IncludeTool/Rules.cs
--- a/Engine/Source/Programs/IncludeTool/IncludeTool/Rules.cs
+++ b/Engine/Source/Programs/IncludeTool/IncludeTool/Rules.cs
@@ -415,6 +415,7 @@ namespace IncludeTool
 			"/Engine/Source/Runtime/MovieScene/Public/MovieSceneFwd.h",
 			"/Engine/Source/Runtime/Core/Public/Internationalization/StringTableCoreFwd.h", // Typedef isn't a forward declaration
 			"/Engine/Source/Runtime/Experimental/Chaos/Public/Chaos/ParticleHandleFwd.h", // invalid forward declaration - 'namespace Chaos'
+			"/Engine/Source/Runtime/Experimental/Chaos/Public/Chaos/GeometryParticlesfwd.h", // invalid forward declaration - 'namespace Chaos'
 		};
 
 		/// <summary>
