Fix IncludeTool parsing of "NoPCHs" value.

#rnx

diff --git a/Engine/Source/Programs/IncludeTool/IncludeTool/BuildTarget.cs b/Engine/Source/Programs/IncludeTool/IncludeTool/BuildTarget.cs
--- a/Engine/Source/Programs/IncludeTool/IncludeTool/BuildTarget.cs
+++ b/Engine/Source/Programs/IncludeTool/IncludeTool/BuildTarget.cs
@@ -41,6 +41,7 @@ namespace IncludeTool
 	enum BuildModulePCHUsage
 	{
 		Default,
+		NoPCHs,
 		NoSharedPCHs,
 		UseSharedPCHs,
 		UseExplicitOrSharedPCHs
