Fix incorrect start index for Substring when finding /FI
#rnx

diff --git a/Engine/Source/Programs/IncludeTool/IncludeTool/CompileEnvironment.cs b/Engine/Source/Programs/IncludeTool/IncludeTool/CompileEnvironment.cs
--- a/Engine/Source/Programs/IncludeTool/IncludeTool/CompileEnvironment.cs
+++ b/Engine/Source/Programs/IncludeTool/IncludeTool/CompileEnvironment.cs
@@ -1,4 +1,4 @@
-﻿// Copyright 1998-2019 Epic Games, Inc. All Rights Reserved.
+// Copyright 1998-2019 Epic Games, Inc. All Rights Reserved.
 
 using IncludeTool.Support;
 using System;
@@ -314,7 +314,7 @@ namespace IncludeTool
 							}
 							else if(Tokens[Idx].StartsWith("/FI"))
 							{
-								Environment.ForceIncludeFiles.Add(new FileReference(Tokens[Idx].Substring(2)));
+								Environment.ForceIncludeFiles.Add(new FileReference(Tokens[Idx].Substring(3)));
 							}
 							else if(Tokens[Idx] == "-include")
 							{
