Update IncludeTool to ignore UE_DEPRECATED_FORGAME instead of DEPRECATED_FORGAME as it was itself deprecated in 4.22
#rnx

diff --git a/Engine/Source/Programs/IncludeTool/IncludeTool/Rules.cs b/Engine/Source/Programs/IncludeTool/IncludeTool/Rules.cs
--- a/Engine/Source/Programs/IncludeTool/IncludeTool/Rules.cs
+++ b/Engine/Source/Programs/IncludeTool/IncludeTool/Rules.cs
@@ -60,7 +60,7 @@ namespace IncludeTool
 
 			return true;
 		}
-
+		
 		/// <summary>
 		/// List of include tokens which are external files, and do not need to be resolved
 		/// </summary>
@@ -117,7 +117,7 @@ namespace IncludeTool
 			"unistd.h",
 			"stdio.h",
 			"unicode/",
-
+			
 			// Android
 			"SLES/OpenSLES_Android.h",
 
@@ -141,7 +141,7 @@ namespace IncludeTool
 			"iphlpapi.h",
 			"Iphlpapi.h",
 			"IcmpAPI.h",
-
+			
 			// Mac
 			"AUEffectBase.h",
 			"Security/Security.h",
@@ -414,7 +414,7 @@ namespace IncludeTool
 			"/Engine/Source/Runtime/Experimental/Chaos/Public/Chaos/ParticleHandleFwd.h", // invalid forward declaration - 'namespace Chaos'
 			"/Engine/Source/Runtime/Experimental/Chaos/Public/Chaos/GeometryParticlesfwd.h", // invalid forward declaration - 'namespace Chaos'
 			"/Engine/Source/Runtime/Experimental/ChaosSolvers/Public/PhysicsProxy/SingleParticlePhysicsProxyFwd.h", // invalid forward declaration - 'namespace Chaos'
-        };
+		};
 
 		/// <summary>
 		/// Gets the flags for a new source file
@@ -529,7 +529,7 @@ namespace IncludeTool
 			{
 				return true;
 			}
-			if(Markup.Type == PreprocessorMarkupType.Define && Markup.Tokens[0].Text == "DEPRECATED_FORGAME")
+			if(Markup.Type == PreprocessorMarkupType.Define && Markup.Tokens[0].Text == "UE_DEPRECATED_FORGAME")
 			{
 				return true;
 			}
@@ -571,18 +571,18 @@ namespace IncludeTool
 			return false;
 		}
 
-		/// <summary>
-		/// Allow overriding whether a symbol should be forward-declared
-		/// </summary>
-		/// <param name="Symbol"></param>
-		/// <returns></returns>
-		public static bool AllowSymbol(string Name)
-		{
-			if(Name == "FNode" || Name == "FFunctionExpression" || Name == "ITextData" || Name == "Rect")
-			{
-				return false;
-			}
-			return true;
-		}
+        /// <summary>
+        /// Allow overriding whether a symbol should be forward-declared
+        /// </summary>
+        /// <param name="Symbol"></param>
+        /// <returns></returns>
+        public static bool AllowSymbol(string Name)
+        {
+            if(Name == "FNode" || Name == "FFunctionExpression" || Name == "ITextData" || Name == "Rect")
+            {
+                return false;
+            }
+            return true;
+        }
 	}
 }
