Still need DEPRECATED_FORGAME as well as UE_DEPRECATED_FORGAME
#rnx

diff --git a/Engine/Source/Programs/IncludeTool/IncludeTool/Rules.cs b/Engine/Source/Programs/IncludeTool/IncludeTool/Rules.cs
--- a/Engine/Source/Programs/IncludeTool/IncludeTool/Rules.cs
+++ b/Engine/Source/Programs/IncludeTool/IncludeTool/Rules.cs
@@ -528,8 +528,7 @@ namespace IncludeTool
 			if(Markup.Type == PreprocessorMarkupType.Define && Markup.Tokens[0].Text == "ONLINE_LOG_PREFIX")
 			{
 				return true;
-			}
-			if(Markup.Type == PreprocessorMarkupType.Define && Markup.Tokens[0].Text == "UE_DEPRECATED_FORGAME")
+			}			if(Markup.Type == PreprocessorMarkupType.Define && (Markup.Tokens[0].Text == "UE_DEPRECATED_FORGAME" || Markup.Tokens[0].Text == "DEPRECATED_FORGAME"))
 			{
 				return true;
 			}
