Remove files that were exceptions to the #include rules, but no longer need to be
#rnx

diff --git a/Engine/Source/Programs/IncludeTool/IncludeTool/Rules.cs b/Engine/Source/Programs/IncludeTool/IncludeTool/Rules.cs
--- a/Engine/Source/Programs/IncludeTool/IncludeTool/Rules.cs
+++ b/Engine/Source/Programs/IncludeTool/IncludeTool/Rules.cs
@@ -389,8 +389,6 @@ namespace IncludeTool
 		/// </summary>
 		static readonly HashSet<string> AllowMultipleFragmentFileNames = new HashSet<string>(StringComparer.OrdinalIgnoreCase)
 		{
-			"/Engine/Source/Runtime/CoreUObject/Public/UObject/UObjectBaseUtility.h",
-			"/Engine/Source/Runtime/CoreUObject/Public/UObject/Class.h",
 			"/Engine/Source/Runtime/Sockets/Private/BSDSockets/SocketSubsystemBSD.h",
 			"/Engine/Source/Runtime/Sockets/Private/BSDSockets/SocketsBSD.h",
 			"/Engine/Plugins/Runtime/PacketHandlers/CompressionComponents/Oodle/Source/OodleHandlerComponent/Public/OodleHandlerComponent.h",
