Addressed IncludeTool warnings for rapidjson

#rnx

diff --git a/Engine/Source/Programs/IncludeTool/IncludeTool/Rules.cs b/Engine/Source/Programs/IncludeTool/IncludeTool/Rules.cs
--- a/Engine/Source/Programs/IncludeTool/IncludeTool/Rules.cs
+++ b/Engine/Source/Programs/IncludeTool/IncludeTool/Rules.cs
@@ -198,6 +198,10 @@ namespace IncludeTool
 			{
 				return true;
 			}
+			if(NormalizedPath.Contains("/thirdparty/rapidjson/"))
+			{
+				return true;
+			}
 			return false;
 		}
 
