Fix IncludeTool warnings for upropery headers.

#rnx

diff --git a/Engine/Source/Programs/IncludeTool/IncludeTool/Rules.cs b/Engine/Source/Programs/IncludeTool/IncludeTool/Rules.cs
--- a/Engine/Source/Programs/IncludeTool/IncludeTool/Rules.cs
+++ b/Engine/Source/Programs/IncludeTool/IncludeTool/Rules.cs
@@ -465,7 +465,12 @@ namespace IncludeTool
 				Flags |= SourceFileFlags.FwdHeader;
 			}
 
-			if(PinnedFileNames.Contains(NormalizedPath))
+			if(NormalizedPath.EndsWith("defineupropertymacros.h"))
+			{
+				Flags |= SourceFileFlags.Inline;
+			}
+
+			if (PinnedFileNames.Contains(NormalizedPath))
 			{
 				Flags = (Flags | SourceFileFlags.Pinned) & ~SourceFileFlags.Standalone;
 			}
