Fixed compile error introduced by CL 8497721 when modifying FramePro headers along with adaptive unity builds.

#error Platform must defined FRAMEPRO_GET_CLOCK_COUNT
and friends.

diff --git a/Engine/Source/Programs/IncludeTool/IncludeTool/Rules.cs b/Engine/Source/Programs/IncludeTool/IncludeTool/Rules.cs
--- a/Engine/Source/Programs/IncludeTool/IncludeTool/Rules.cs
+++ b/Engine/Source/Programs/IncludeTool/IncludeTool/Rules.cs
@@ -194,7 +194,7 @@ namespace IncludeTool
 			{
 				return true;
 			}
-			if(NormalizedPath.EndsWith("/recastmesh.cpp") || NormalizedPath.EndsWith("/recastfilter.cpp") || NormalizedPath.EndsWith("/recastcontour.cpp") || NormalizedPath.EndsWith("/framepro.h") || NormalizedPath.EndsWith("/framepro.cpp") || NormalizedPath.EndsWith("/sqlite3.h") || NormalizedPath.EndsWith("/sqlite3.inl") || NormalizedPath.EndsWith("/vorbis_stream_encoder.h") || NormalizedPath.EndsWith("/integral_types.h"))
+			if(NormalizedPath.EndsWith("/recastmesh.cpp") || NormalizedPath.EndsWith("/recastfilter.cpp") || NormalizedPath.EndsWith("/recastcontour.cpp") || NormalizedPath.EndsWith("/framepro.h") || NormalizedPath.EndsWith("/framepro.cpp") || NormalizedPath.EndsWith("/frameproue4.h") || NormalizedPath.EndsWith("/frameproue4.cpp") || NormalizedPath.EndsWith("/sqlite3.h") || NormalizedPath.EndsWith("/sqlite3.inl") || NormalizedPath.EndsWith("/vorbis_stream_encoder.h") || NormalizedPath.EndsWith("/integral_types.h"))
 			{
 				return true;
 			}
diff --git a/Engine/Source/Runtime/Core/Public/FramePro/FrameProUE4.h b/Engine/Source/Runtime/Core/Public/FramePro/FrameProUE4.h
--- a/Engine/Source/Runtime/Core/Public/FramePro/FrameProUE4.h
+++ b/Engine/Source/Runtime/Core/Public/FramePro/FrameProUE4.h
@@ -22,8 +22,6 @@
 	Add FramePro.cpp to your project to allow FramePro to communicate with your application.
 */
 
-#pragma once
-
 //------------------------------------------------------------------------
 #include "FramePro.h"
 
