IncludeTool: Skip over "__has_warning" preprocessor tests.

diff --git a/Engine/Source/Programs/IncludeTool/IncludeTool/PreprocessorExpression.cs b/Engine/Source/Programs/IncludeTool/IncludeTool/PreprocessorExpression.cs
--- a/Engine/Source/Programs/IncludeTool/IncludeTool/PreprocessorExpression.cs
+++ b/Engine/Source/Programs/IncludeTool/IncludeTool/PreprocessorExpression.cs
@@ -341,6 +341,7 @@ namespace IncludeTool
 					throw new NotImplementedException();
 				case "__has_builtin":
 				case "__has_feature":
+				case "__has_warning":
 					if(Tokens[Idx + 1].Text != "(" || Tokens[Idx + 3].Text != ")")
 					{
 						throw new NotImplementedException();
