Copyright notice updates

diff --git a/Engine/Source/Programs/IncludeTool/IncludeTool/Support/LineBasedTextWriter.cs b/Engine/Source/Programs/IncludeTool/IncludeTool/Support/LineBasedTextWriter.cs
--- a/Engine/Source/Programs/IncludeTool/IncludeTool/Support/LineBasedTextWriter.cs
+++ b/Engine/Source/Programs/IncludeTool/IncludeTool/Support/LineBasedTextWriter.cs
@@ -1,3 +1,5 @@
+// Copyright 1998-2019 Epic Games, Inc. All Rights Reserved.
+
 using System;
 using System.Collections.Generic;
 using System.Linq;
diff --git a/Engine/Source/Programs/IncludeTool/IncludeTool/Support/LineBasedTextWriterWrapper.cs b/Engine/Source/Programs/IncludeTool/IncludeTool/Support/LineBasedTextWriterWrapper.cs
--- a/Engine/Source/Programs/IncludeTool/IncludeTool/Support/LineBasedTextWriterWrapper.cs
+++ b/Engine/Source/Programs/IncludeTool/IncludeTool/Support/LineBasedTextWriterWrapper.cs
@@ -1,3 +1,5 @@
+// Copyright 1998-2019 Epic Games, Inc. All Rights Reserved.
+
 using System;
 using System.Collections.Generic;
 using System.IO;
diff --git a/Engine/Source/Programs/RoboMerge/public/js/graphviz-worker.js b/Engine/Source/Programs/RoboMerge/public/js/graphviz-worker.js
--- a/Engine/Source/Programs/RoboMerge/public/js/graphviz-worker.js
+++ b/Engine/Source/Programs/RoboMerge/public/js/graphviz-worker.js
@@ -1,4 +1,4 @@
-// Copyright 1998-2016 Epic Games, Inc. All Rights Reserved.
+// Copyright 1998-2019 Epic Games, Inc. All Rights Reserved.
 
 importScripts('/js/graphviz-asm.js');
 importScripts('/js/graphviz.js');
diff --git a/Engine/Source/Programs/RoboMerge/public/js/login.js b/Engine/Source/Programs/RoboMerge/public/js/login.js
--- a/Engine/Source/Programs/RoboMerge/public/js/login.js
+++ b/Engine/Source/Programs/RoboMerge/public/js/login.js
@@ -1,4 +1,4 @@
-// Copyright 1998-2018 Epic Games, Inc. All Rights Reserved.
+// Copyright 1998-2019 Epic Games, Inc. All Rights Reserved.
 "use strict"
 
 const COOKIE_REGEX = /\s*(.*?)=(.*)/;
diff --git a/Engine/Source/Programs/RoboMerge/public/js/main.js b/Engine/Source/Programs/RoboMerge/public/js/main.js
--- a/Engine/Source/Programs/RoboMerge/public/js/main.js
+++ b/Engine/Source/Programs/RoboMerge/public/js/main.js
@@ -1,4 +1,4 @@
-// Copyright 1998-2018 Epic Games, Inc. All Rights Reserved.
+// Copyright 1998-2019 Epic Games, Inc. All Rights Reserved.
 "use strict"
 
 function promptFor(map) {
