IncludeTool support for clang __has_builtin

diff --git a/Engine/Source/Programs/IncludeTool/IncludeTool/PreprocessorExpression.cs b/Engine/Source/Programs/IncludeTool/IncludeTool/PreprocessorExpression.cs
--- a/Engine/Source/Programs/IncludeTool/IncludeTool/PreprocessorExpression.cs
+++ b/Engine/Source/Programs/IncludeTool/IncludeTool/PreprocessorExpression.cs
@@ -339,6 +339,7 @@ namespace IncludeTool
 					throw new NotImplementedException();
 				case "alignof":
 					throw new NotImplementedException();
+				case "__has_builtin":
 				case "__has_feature":
 					if(Tokens[Idx + 1].Text != "(" || Tokens[Idx + 3].Text != ")")
 					{
