Improve IncludeTool error output when encountering unbalanced curly braces or conditional pre-processor macro blocks.

diff --git a/Engine/Source/Programs/IncludeTool/IncludeTool/SymbolTable.cs b/Engine/Source/Programs/IncludeTool/IncludeTool/SymbolTable.cs
--- a/Engine/Source/Programs/IncludeTool/IncludeTool/SymbolTable.cs
+++ b/Engine/Source/Programs/IncludeTool/IncludeTool/SymbolTable.cs
@@ -251,7 +251,10 @@ namespace IncludeTool
 					// Make sure they were all consistent
 					if(NewScopes.Any(x => x != NewScopes[0]))
 					{
-						throw new UnbalancedScopeException(String.Format("Unbalanced scope depth between conditional block derivations while parsing {0}", File.Location));
+						throw new UnbalancedScopeException(String.Format("Unbalanced scope depth between conditional block derivations while parsing {0} at pos {1}:{2} with end pos {3}:{4}",
+                            File.Location,
+                            Markup.Location.LineIdx + 1, Markup.Location.ColumnIdx,
+                            Markup.EndLocation.LineIdx + 1, Markup.EndLocation.ColumnIdx));
 					}
 
 					// Update the scope
@@ -293,7 +296,10 @@ namespace IncludeTool
 						}
 						else if(Reader.Current.Text == "}" && --Scope < 0)
 						{
-							throw new UnbalancedScopeException(String.Format("Unbalanced '}}' while parsing '{0}'", File.Location));
+							throw new UnbalancedScopeException(String.Format("Unbalanced '}}' while parsing '{0}' at pos {1}:{2} with end pos {3}:{4}",
+                                File.Location,
+                                Reader.TokenLocation.LineIdx+1, Reader.TokenLocation.ColumnIdx,
+                                Reader.TokenEndLocation.LineIdx + 1, Reader.TokenEndLocation.ColumnIdx));
 						}
 
 						// Move to the next token
