IncludeTool: Skip over "__has_include" preprocessor statements. These statements take the form __has_include(x) where x is an quoted or <>-surrounded path. Thus the number of tokens is variable.

#rnx

diff --git a/Engine/Source/Programs/IncludeTool/IncludeTool/PreprocessorExpression.cs b/Engine/Source/Programs/IncludeTool/IncludeTool/PreprocessorExpression.cs
--- a/Engine/Source/Programs/IncludeTool/IncludeTool/PreprocessorExpression.cs
+++ b/Engine/Source/Programs/IncludeTool/IncludeTool/PreprocessorExpression.cs
@@ -348,6 +348,20 @@ namespace IncludeTool
 					}
 					Idx += 4;
 					return 0;
+				case "__has_include":
+					if(Tokens[Idx + 1].Text == "(")
+					{
+						Idx += 4; // 4 = __has_inc, (, "" or <>
+						for (; Idx < Tokens.Count; ++Idx)
+						{
+							if (Tokens[Idx].Text == ")")
+							{
+								++Idx;
+								return 0;
+							}
+						}
+					}
+					throw new NotImplementedException();
 				case "__building_module":
 					if(Tokens[Idx + 1].Text != "(" || Tokens[Idx + 3].Text != ")")
 					{
