- Added lz4.h to be marked as External header so IncludeTool analysis won't report inconsistent preprocessor settings (lz4hc.h includes twice lz4.h with different defines set)
#rnx

diff --git a/Engine/Source/Programs/IncludeTool/IncludeTool/Rules.cs b/Engine/Source/Programs/IncludeTool/IncludeTool/Rules.cs
--- a/Engine/Source/Programs/IncludeTool/IncludeTool/Rules.cs
+++ b/Engine/Source/Programs/IncludeTool/IncludeTool/Rules.cs
@@ -202,6 +202,10 @@ namespace IncludeTool
 			{
 				return true;
 			}
+			if (NormalizedPath.EndsWith("/lz4.h") || NormalizedPath.EndsWith("/lz4hc.h"))
+			{
+				return true;
+			}
 			return false;
 		}
 
