- Undid the IncludeTool change because it caused another issue, solving it another way

diff --git a/Engine/Source/Programs/IncludeTool/IncludeTool/Rules.cs b/Engine/Source/Programs/IncludeTool/IncludeTool/Rules.cs
--- a/Engine/Source/Programs/IncludeTool/IncludeTool/Rules.cs
+++ b/Engine/Source/Programs/IncludeTool/IncludeTool/Rules.cs
@@ -202,10 +202,6 @@ namespace IncludeTool
 			{
 				return true;
 			}
-			if (NormalizedPath.EndsWith("/lz4.h") || NormalizedPath.EndsWith("/lz4hc.h"))
-			{
-				return true;
-			}
 			return false;
 		}
 
diff --git a/Engine/Source/Runtime/Core/Public/Compression/lz4.h b/Engine/Source/Runtime/Core/Public/Compression/lz4.h
--- a/Engine/Source/Runtime/Core/Public/Compression/lz4.h
+++ b/Engine/Source/Runtime/Core/Public/Compression/lz4.h
@@ -32,6 +32,9 @@
     - LZ4 homepage : http://www.lz4.org
     - LZ4 source repository : [REDACTED_PROFILE_URL]/lz4
 */
+
+#pragma once
+
 #include "CoreMinimal.h"
 
 #if defined (__cplusplus)
