- Final fix for lz4 IncludeTool warnings

diff --git a/Engine/Source/Programs/IncludeTool/IncludeTool/Rules.cs b/Engine/Source/Programs/IncludeTool/IncludeTool/Rules.cs
--- a/Engine/Source/Programs/IncludeTool/IncludeTool/Rules.cs
+++ b/Engine/Source/Programs/IncludeTool/IncludeTool/Rules.cs
@@ -202,6 +202,10 @@ namespace IncludeTool
 			{
 				return true;
 			}
+			if (NormalizedPath.EndsWith("/lz4.h") || NormalizedPath.EndsWith("/lz4hc.h") || NormalizedPath.EndsWith("/lz4.cpp"))
+			{
+				return true;
+			}
 			return false;
 		}
 
diff --git a/Engine/Source/Runtime/Core/Public/Compression/lz4.h b/Engine/Source/Runtime/Core/Public/Compression/lz4.h
--- a/Engine/Source/Runtime/Core/Public/Compression/lz4.h
+++ b/Engine/Source/Runtime/Core/Public/Compression/lz4.h
@@ -33,8 +33,6 @@
     - LZ4 source repository : [REDACTED_PROFILE_URL]/lz4
 */
 
-#include "CoreMinimal.h"
-
 #if defined (__cplusplus)
 extern "C" {
 #endif
