- Yet more lz4 fixes, for NonUnity, this time
#rnx

diff --git a/Engine/Source/Programs/IncludeTool/IncludeTool/Program.cs b/Engine/Source/Programs/IncludeTool/IncludeTool/Program.cs
--- a/Engine/Source/Programs/IncludeTool/IncludeTool/Program.cs
+++ b/Engine/Source/Programs/IncludeTool/IncludeTool/Program.cs
@@ -1026,7 +1026,7 @@ namespace IncludeTool
 				{
 					foreach(PreprocessorMarkup Markup in File.Markup)
 					{
-						if(Markup.IncludedFile != null && (Markup.IncludedFile.Flags & SourceFileFlags.External) == 0 && !File.Location.FullName.Contains("PhyaLib") && !File.Location.GetFileName().Contains("Recast") && !File.Location.GetFileName().Contains("FramePro") && !File.Location.FullName.Contains("libSampleRate"))
+						if(Markup.IncludedFile != null && (Markup.IncludedFile.Flags & SourceFileFlags.External) == 0 && !File.Location.FullName.Contains("PhyaLib") && !File.Location.GetFileName().Contains("Recast") && !File.Location.GetFileName().Contains("FramePro") && !File.Location.FullName.Contains("libSampleRate") && !File.Location.FullName.Contains("lz4"))
 						{
 							Log.WriteLine("{0}({1}): warning: external file is including non-external file ('{2}')", IncludedFile.Location.FullName, Markup.Location.LineIdx + 1, Markup.IncludedFile.Location);
 						}
diff --git a/Engine/Source/Programs/IncludeTool/IncludeTool/Rules.cs b/Engine/Source/Programs/IncludeTool/IncludeTool/Rules.cs
--- a/Engine/Source/Programs/IncludeTool/IncludeTool/Rules.cs
+++ b/Engine/Source/Programs/IncludeTool/IncludeTool/Rules.cs
@@ -202,7 +202,7 @@ namespace IncludeTool
 			{
 				return true;
 			}
-			if (NormalizedPath.EndsWith("/lz4.h") || NormalizedPath.EndsWith("/lz4hc.h") || NormalizedPath.EndsWith("/lz4.cpp"))
+			if (NormalizedPath.EndsWith("/lz4.h") || NormalizedPath.EndsWith("/lz4hc.h") || NormalizedPath.EndsWith("/lz4.cpp") || NormalizedPath.EndsWith("/lz4hc.cpp"))
 			{
 				return true;
 			}
diff --git a/Engine/Source/Runtime/Core/Public/Compression/lz4.h b/Engine/Source/Runtime/Core/Public/Compression/lz4.h
--- a/Engine/Source/Runtime/Core/Public/Compression/lz4.h
+++ b/Engine/Source/Runtime/Core/Public/Compression/lz4.h
@@ -33,6 +33,8 @@
     - LZ4 source repository : [REDACTED_PROFILE_URL]/lz4
 */
 
+#include "CoreMinimal.h"
+
 #if defined (__cplusplus)
 extern "C" {
 #endif
@@ -557,7 +559,7 @@ LZ4LIB_STATIC_API void LZ4_attach_dictionary(LZ4_stream_t* workingStream, const
 #define LZ4_HASHTABLESIZE (1 << LZ4_MEMORY_USAGE)
 #define LZ4_HASH_SIZE_U32 (1 << LZ4_HASHLOG)       /* required as macro for static allocation */
 
-#if 0// defined(__cplusplus) || (defined (__STDC_VERSION__) && (__STDC_VERSION__ >= 199901L) /* C99 */)
+#if defined(__cplusplus) || (defined (__STDC_VERSION__) && (__STDC_VERSION__ >= 199901L) /* C99 */)
 
 typedef struct LZ4_stream_t_internal LZ4_stream_t_internal;
 struct LZ4_stream_t_internal {
