Added fwd header to IgnoreFwdHeaders

Ben.Marsh

diff --git a/Engine/Source/Programs/IncludeTool/IncludeTool/Rules.cs b/Engine/Source/Programs/IncludeTool/IncludeTool/Rules.cs
--- a/Engine/Source/Programs/IncludeTool/IncludeTool/Rules.cs
+++ b/Engine/Source/Programs/IncludeTool/IncludeTool/Rules.cs
@@ -419,6 +419,7 @@ namespace IncludeTool
 			"/Engine/Source/Runtime/Core/Public/Internationalization/StringTableCoreFwd.h", // Typedef isn't a forward declaration
 			"/Engine/Source/Runtime/Experimental/Chaos/Public/Chaos/ParticleHandleFwd.h", // invalid forward declaration - 'namespace Chaos'
 			"/Engine/Source/Runtime/Experimental/Chaos/Public/Chaos/GeometryParticlesfwd.h", // invalid forward declaration - 'namespace Chaos'
+			"/Engine/Source/Runtime/Experimental/Chaos/Public/Chaos/PBDRigidsEvolutionFwd.h", // invalid forward declaration - 'namespace Chaos'
 			"/Engine/Source/Runtime/Experimental/ChaosSolvers/Public/PhysicsProxy/SingleParticlePhysicsProxyFwd.h", // invalid forward declaration - 'namespace Chaos'
 		};
 
