Editgrated 14042502 from Dev-CharacterTech to UE5/Main. Fixes IncludeTool static analysis errors.

diff --git a/Engine/Source/Programs/IncludeTool/IncludeTool/Rules.cs b/Engine/Source/Programs/IncludeTool/IncludeTool/Rules.cs
--- a/Engine/Source/Programs/IncludeTool/IncludeTool/Rules.cs
+++ b/Engine/Source/Programs/IncludeTool/IncludeTool/Rules.cs
@@ -560,6 +560,14 @@ namespace IncludeTool
 			{
 				return true;
 			}
+			if (Markup.Type == PreprocessorMarkupType.Define && (Markup.Tokens[0].Text == "RLAPI" || Markup.Tokens[0].Text == "DNAAPI" || Markup.Tokens[0].Text == "TRIOAPI" || Markup.Tokens[0].Text == "SCAPI" || Markup.Tokens[0].Text == "PMAAPI"))
+			{
+				return true;
+			}
+			if (Markup.Type == PreprocessorMarkupType.Elif && Markup.Tokens.Count == 4 && Markup.Tokens[2].Text == "RL_SHARED")
+			{
+				return true;
+			}
 			return false;
 		}
 
