Add UnrealWidgetFwd.h temporarily to IncludeTool's Rules.cs to fix SA

The real fix will be to rename this file but since this is for M2 I will do that afterwards in UE5-Main

#rnx

diff --git a/Engine/Source/Programs/IncludeTool/IncludeTool/Rules.cs b/Engine/Source/Programs/IncludeTool/IncludeTool/Rules.cs
--- a/Engine/Source/Programs/IncludeTool/IncludeTool/Rules.cs
+++ b/Engine/Source/Programs/IncludeTool/IncludeTool/Rules.cs
@@ -413,6 +413,7 @@ namespace IncludeTool
 		{
 			"/Engine/Source/Runtime/Core/Public/Internationalization/TextNamespaceFwd.h",
 			"/Engine/Source/Editor/SceneOutliner/Public/SceneOutlinerFwd.h",
+			"/Engine/Source/Editor/UnrealEd/Public/UnrealWidgetFwd.h", // error: invalid forward declaration - 'enum ECoordSystem'
 			"/Engine/Source/Runtime/SlateCore/Public/Fonts/ShapedTextFwd.h", // Typedef isn't a forward declaration
 			"/Engine/Source/Runtime/Slate/Public/Framework/Text/ShapedTextCacheFwd.h", // Typedef isn't a forward declaration
 			"/Engine/Source/Runtime/MovieScene/Public/MovieSceneFwd.h",
