Added an exception for the InterchangeEngineFwd.h in IncludeTool since it doesn't handle namespaces in Fwd files.

diff --git a/Engine/Source/Programs/IncludeTool/IncludeTool/Rules.cs b/Engine/Source/Programs/IncludeTool/IncludeTool/Rules.cs
--- a/Engine/Source/Programs/IncludeTool/IncludeTool/Rules.cs
+++ b/Engine/Source/Programs/IncludeTool/IncludeTool/Rules.cs
@@ -426,6 +426,7 @@ namespace IncludeTool
 			"/Engine/Source/Runtime/Experimental/Chaos/Public/Chaos/PBDRigidsEvolutionFwd.h", // invalid forward declaration - 'namespace Chaos'
 			"/Engine/Source/Runtime/Experimental/Chaos/Public/PhysicsProxy/JointConstraintProxyFwd.h", // invalid forward declaration - 'namespace Chaos'
 			"/Engine/Source/Runtime/Experimental/Chaos/Public/PhysicsProxy/SingleParticlePhysicsProxyFwd.h", // invalid forward declaration - 'namespace Chaos'
+			"/Engine/Source/Runtime/Interchange/Engine/Public/InterchangeEngineFwd.h", // invalid forward declaration - 'namespace UE'
 		};
 
 		/// <summary>
