Fixed exclusion path for moved file

#rnx

diff --git a/Engine/Source/Programs/IncludeTool/IncludeTool/Rules.cs b/Engine/Source/Programs/IncludeTool/IncludeTool/Rules.cs
--- a/Engine/Source/Programs/IncludeTool/IncludeTool/Rules.cs
+++ b/Engine/Source/Programs/IncludeTool/IncludeTool/Rules.cs
@@ -413,7 +413,7 @@ namespace IncludeTool
 		{
 			"/Engine/Source/Runtime/Core/Public/Internationalization/TextNamespaceFwd.h",
 			"/Engine/Source/Editor/SceneOutliner/Public/SceneOutlinerFwd.h",
-			"/Engine/Source/Editor/UnrealEd/Public/UnrealWidgetFwd.h", // error: invalid forward declaration - 'enum ECoordSystem'
+			"/Engine/Source/Editor/EditorFramework/Public/UnrealWidgetFwd.h", // error: invalid forward declaration - 'enum ECoordSystem'
 			"/Engine/Source/Runtime/SlateCore/Public/Fonts/ShapedTextFwd.h", // Typedef isn't a forward declaration
 			"/Engine/Source/Runtime/Slate/Public/Framework/Text/ShapedTextCacheFwd.h", // Typedef isn't a forward declaration
 			"/Engine/Source/Runtime/MovieScene/Public/MovieSceneFwd.h",
