Convert IncludeTool to NET Core 3.1.

#rnx

diff --git a/Engine/Source/Programs/IncludeTool/IncludeTool/IncludeTool.csproj b/Engine/Source/Programs/IncludeTool/IncludeTool/IncludeTool.csproj
--- a/Engine/Source/Programs/IncludeTool/IncludeTool/IncludeTool.csproj
+++ b/Engine/Source/Programs/IncludeTool/IncludeTool/IncludeTool.csproj
@@ -1,151 +1,60 @@
-﻿<?xml version="1.0" encoding="utf-8"?>
-<Project ToolsVersion="12.0" DefaultTargets="Build" xmlns="http://schemas.microsoft.com/developer/msbuild/2003">
-  <Import Project="$(MSBuildExtensionsPath)\$(MSBuildToolsVersion)\Microsoft.Common.props" Condition="Exists('$(MSBuildExtensionsPath)\$(MSBuildToolsVersion)\Microsoft.Common.props')" />
+﻿<Project Sdk="Microsoft.NET.Sdk">
   <PropertyGroup>
-    <Configuration Condition=" '$(Configuration)' == '' ">Debug</Configuration>
-    <Platform Condition=" '$(Platform)' == '' ">AnyCPU</Platform>
-    <ProjectGuid>{CB278021-B7F7-4DCD-9D46-DCCBE6A4B3F9}</ProjectGuid>
+    <TargetFramework>netcoreapp3.1</TargetFramework>
+    <Configuration Condition=" '$(Configuration)' == '' ">Development</Configuration>
     <OutputType>Exe</OutputType>
-    <AppDesignerFolder>Properties</AppDesignerFolder>
-    <RootNamespace>IncludeTool</RootNamespace>
+    <AppendTargetFrameworkToOutputPath>false</AppendTargetFrameworkToOutputPath>
+    <GenerateAssemblyInfo>false</GenerateAssemblyInfo>
+    <GenerateTargetFrameworkAttribute>false</GenerateTargetFrameworkAttribute>
+    <Configurations>Debug;Release;Development</Configurations>
     <AssemblyName>IncludeTool</AssemblyName>
-    <TargetFrameworkVersion>v4.5</TargetFrameworkVersion>
-    <FileAlignment>512</FileAlignment>
-    <TargetFrameworkProfile />
+    <RootNamespace>IncludeTool</RootNamespace>
+    <SatelliteResourceLanguages>en</SatelliteResourceLanguages> <!-- remove non english resource languages -->
+    <AppendRuntimeIdentifierToOutputPath>false</AppendRuntimeIdentifierToOutputPath>
+    <IsWindows Condition="'$([System.Runtime.InteropServices.RuntimeInformation]::IsOSPlatform($([System.Runtime.InteropServices.OSPlatform]::Windows)))' == 'true'">true</IsWindows> 
+    <IsOSX Condition="'$([System.Runtime.InteropServices.RuntimeInformation]::IsOSPlatform($([System.Runtime.InteropServices.OSPlatform]::OSX)))' == 'true'">true</IsOSX> 
+    <IsLinux Condition="'$([System.Runtime.InteropServices.RuntimeInformation]::IsOSPlatform($([System.Runtime.InteropServices.OSPlatform]::Linux)))' == 'true'">true</IsLinux> 
+    <DebugType>pdbonly</DebugType> 
+    <Nullable>disable</Nullable>
   </PropertyGroup>
-  <PropertyGroup Condition="'$(Configuration)|$(Platform)' == 'Debug|x64'">
-    <DebugSymbols>true</DebugSymbols>
-    <OutputPath>..\..\..\..\Binaries\DotNET\</OutputPath>
-    <DefineConstants>DEBUG;TRACE</DefineConstants>
-    <DebugType>full</DebugType>
-    <PlatformTarget>x64</PlatformTarget>
-    <ErrorReport>prompt</ErrorReport>
-    <Prefer32Bit>true</Prefer32Bit>
+  <Choose>
+    <When Condition="'$(IsOSX)'=='true'">
+      <PropertyGroup>
+        <!-- Set runtie identifier for osx so that it produces a app host executable (consistency with the other platforms) -->
+        <RuntimeIdentifier>osx-x64</RuntimeIdentifier>
+      </PropertyGroup>
+    </When>
+  </Choose>
+  <PropertyGroup Condition=" '$(Configuration)|$(Platform)' == 'Development|AnyCPU' ">
+    <Optimize>true</Optimize>
+    <OutputPath>..\..\..\Binaries\DotNET\IncludeTool</OutputPath>
   </PropertyGroup>
-  <PropertyGroup Condition="'$(Configuration)|$(Platform)' == 'Release|x64'">
-    <OutputPath>..\..\..\..\Binaries\DotNET\</OutputPath>
-    <DefineConstants>TRACE</DefineConstants>
-    <Optimize>false</Optimize>
-    <DebugType>pdbonly</DebugType>
-    <PlatformTarget>x64</PlatformTarget>
-    <ErrorReport>prompt</ErrorReport>
-    <Prefer32Bit>true</Prefer32Bit>
+  <PropertyGroup Condition=" '$(Configuration)|$(Platform)' == 'Release|AnyCPU' ">
+    <Optimize>true</Optimize>
+    <OutputPath>..\..\..\Binaries\DotNET\IncludeTool</OutputPath>
   </PropertyGroup>
-  <PropertyGroup />
-  <PropertyGroup Condition="'$(Configuration)|$(Platform)' == 'Debug|AnyCPU'">
+  <PropertyGroup Condition=" '$(Configuration)|$(Platform)' == 'Debug|AnyCPU' ">
+    <OutputPath>..\..\..\..\Binaries\DotNET\IncludeTool</OutputPath>
     <DebugSymbols>true</DebugSymbols>
-    <OutputPath>..\..\..\..\Binaries\DotNET\</OutputPath>
-    <DefineConstants>DEBUG;TRACE</DefineConstants>
     <DebugType>full</DebugType>
-    <PlatformTarget>AnyCPU</PlatformTarget>
-    <ErrorReport>prompt</ErrorReport>
-    <Prefer32Bit>false</Prefer32Bit>
-  </PropertyGroup>
-  <PropertyGroup Condition="'$(Configuration)|$(Platform)' == 'Release|AnyCPU'">
-    <OutputPath>..\..\..\..\Binaries\DotNET\</OutputPath>
-    <DefineConstants>TRACE</DefineConstants>
-    <DebugType>pdbonly</DebugType>
-    <PlatformTarget>AnyCPU</PlatformTarget>
-    <ErrorReport>prompt</ErrorReport>
-    <Prefer32Bit>false</Prefer32Bit>
-    <Optimize>true</Optimize>
   </PropertyGroup>
   <ItemGroup>
-    <Reference Include="System" />
-    <Reference Include="System.Core" />
-    <Reference Include="System.Xml.Linq" />
-    <Reference Include="System.Data.DataSetExtensions" />
-    <Reference Include="Microsoft.CSharp" />
-    <Reference Include="System.Data" />
-    <Reference Include="System.Xml" />
-  </ItemGroup>
-  <ItemGroup>
-    <Compile Include="..\..\DotNETCommon\DotNETUtilities\FastJSON\Formatter.cs">
-      <Link>FastJSON\Formatter.cs</Link>
-    </Compile>
-    <Compile Include="..\..\DotNETCommon\DotNETUtilities\FastJSON\Getters.cs">
-      <Link>FastJSON\Getters.cs</Link>
-    </Compile>
-    <Compile Include="..\..\DotNETCommon\DotNETUtilities\FastJSON\JSON.cs">
-      <Link>FastJSON\JSON.cs</Link>
-    </Compile>
-    <Compile Include="..\..\DotNETCommon\DotNETUtilities\FastJSON\JsonParser.cs">
-      <Link>FastJSON\JsonParser.cs</Link>
-    </Compile>
-    <Compile Include="..\..\DotNETCommon\DotNETUtilities\FastJSON\JsonSerializer.cs">
-      <Link>FastJSON\JsonSerializer.cs</Link>
-    </Compile>
-    <Compile Include="..\..\DotNETCommon\DotNETUtilities\FastJSON\Reflection.cs">
-      <Link>FastJSON\Reflection.cs</Link>
-    </Compile>
-    <Compile Include="..\..\DotNETCommon\DotNETUtilities\FastJSON\SafeDictionary.cs">
-      <Link>FastJSON\SafeDictionary.cs</Link>
-    </Compile>
-    <Compile Include="BuildTarget.cs" />
-    <Compile Include="ManagedProcess.cs" />
-    <Compile Include="OutputFile.cs" />
-    <Compile Include="Reports\ComplexityReport.cs" />
-    <Compile Include="Reports\FragmentTimingReport.cs" />
-    <Compile Include="Reports\PchReport.cs" />
-    <Compile Include="Reports\SymbolReport.cs" />
-    <Compile Include="SequenceProbe.cs" />
-    <Compile Include="IncludeCycle.cs" />
-    <Compile Include="CompileEnvironment.cs" />
-    <Compile Include="PreprocessorExpression.cs" />
-    <Compile Include="Preprocessor.cs" />
-    <Compile Include="PreprocessorMarkup.cs" />
-    <Compile Include="Properties\Resources.Designer.cs">
-      <AutoGen>True</AutoGen>
-      <DesignTime>True</DesignTime>
-      <DependentUpon>Resources.resx</DependentUpon>
-    </Compile>
-    <Compile Include="Rules.cs" />
-    <Compile Include="Program.cs" />
-    <Compile Include="Properties\AssemblyInfo.cs" />
-    <Compile Include="SequenceWorker.cs" />
-    <Compile Include="SourceFile.cs" />
-    <Compile Include="SourceFragment.cs" />
-    <Compile Include="Support\BinarySerializer.cs" />
-    <Compile Include="Support\BufferedTextWriter.cs" />
-    <Compile Include="Support\FileFilter.cs" />
-    <Compile Include="Support\FileSystemReference.cs" />
-    <Compile Include="Support\HashList.cs" />
-    <Compile Include="Support\JsonObject.cs" />
-    <Compile Include="Support\JsonWriter.cs" />
-    <Compile Include="Support\LineBasedTextWriter.cs" />
-    <Compile Include="Support\LineBasedTextWriterWrapper.cs" />
-    <Compile Include="Support\LogWriter.cs" />
-    <Compile Include="Support\ManagedTask.cs" />
-    <Compile Include="Support\MultiValueDictionary.cs" />
-    <Compile Include="Support\CommandLine.cs" />
-    <Compile Include="Support\TextBuffer.cs" />
-    <Compile Include="Support\ThreadPoolWorkQueue.cs" />
-    <Compile Include="Support\Utility.cs" />
-    <Compile Include="SymbolTable.cs" />
-    <Compile Include="Token.cs" />
-    <Compile Include="TokenReader.cs" />
-    <Compile Include="Workspace.cs" />
-  </ItemGroup>
-  <ItemGroup>
-    <None Include="App.config" />
-  </ItemGroup>
-  <ItemGroup>
-    <EmbeddedResource Include="Properties\Resources.resx">
-      <Generator>ResXFileCodeGenerator</Generator>
-      <LastGenOutput>Resources.Designer.cs</LastGenOutput>
-    </EmbeddedResource>
+    <PackageReference Include="Microsoft.CodeAnalysis.CSharp" Version="3.7.0" />
+    <PackageReference Include="Microsoft.Win32.Registry" Version="4.7.0" />
+    <PackageReference Include="System.Management" Version="4.7.0" />
+    <PackageReference Include="System.Security.Cryptography.Csp" Version="4.3.0" />
+    <PackageReference Include="System.ServiceProcess.ServiceController" Version="4.7.0" />
+
+    <!-- Required by Ionic -->
+    <PackageReference Include="System.Security.Permissions" Version="4.7.0" />
+    <PackageReference Include="System.Text.Encoding.CodePages" Version="4.7.1" />
+
+    <!-- Required for Hololens -->
+    <PackageReference Include="System.Reflection.MetadataLoadContext" Version="4.7.2" /> 
   </ItemGroup>
   <ItemGroup>
-    <Content Include="Config\Prelude.txt">
-      <CopyToOutputDirectory>PreserveNewest</CopyToOutputDirectory>
-    </Content>
+    <Reference Include="fastJSON">
+      <HintPath>..\..\..\..\Binaries\ThirdParty\FastJson\netstandard2.0\FastJson.dll</HintPath>
+    </Reference>
   </ItemGroup>
-  <Import Project="$(MSBuildToolsPath)\Microsoft.CSharp.targets" />
-  <!-- To modify your build process, add your task inside one of the targets below and uncomment it. 
-       Other similar extension points exist, see Microsoft.Common.targets.
-  <Target Name="BeforeBuild">
-  </Target>
-  <Target Name="AfterBuild">
-  </Target>
-  -->
 </Project>
\ No newline at end of file
