Fix IncludeTool output paths.

#rnx

diff --git a/Engine/Source/Programs/IncludeTool/IncludeTool/IncludeTool.csproj b/Engine/Source/Programs/IncludeTool/IncludeTool/IncludeTool.csproj
--- a/Engine/Source/Programs/IncludeTool/IncludeTool/IncludeTool.csproj
+++ b/Engine/Source/Programs/IncludeTool/IncludeTool/IncludeTool.csproj
@@ -1,4 +1,4 @@
-﻿<Project Sdk="Microsoft.NET.Sdk">
+<Project Sdk="Microsoft.NET.Sdk">
   <PropertyGroup>
     <TargetFramework>netcoreapp3.1</TargetFramework>
     <Configuration Condition=" '$(Configuration)' == '' ">Development</Configuration>
@@ -27,11 +27,11 @@
   </Choose>
   <PropertyGroup Condition=" '$(Configuration)|$(Platform)' == 'Development|AnyCPU' ">
     <Optimize>true</Optimize>
-    <OutputPath>..\..\..\Binaries\DotNET\IncludeTool</OutputPath>
+    <OutputPath>..\..\..\..\Binaries\DotNET\IncludeTool</OutputPath>
   </PropertyGroup>
   <PropertyGroup Condition=" '$(Configuration)|$(Platform)' == 'Release|AnyCPU' ">
     <Optimize>true</Optimize>
-    <OutputPath>..\..\..\Binaries\DotNET\IncludeTool</OutputPath>
+    <OutputPath>..\..\..\..\Binaries\DotNET\IncludeTool</OutputPath>
   </PropertyGroup>
   <PropertyGroup Condition=" '$(Configuration)|$(Platform)' == 'Debug|AnyCPU' ">
     <OutputPath>..\..\..\..\Binaries\DotNET\IncludeTool</OutputPath>
