Fix working directory for IncludeTool.

#rnx

diff --git a/Engine/Source/Programs/IncludeTool/IncludeTool/Program.cs b/Engine/Source/Programs/IncludeTool/IncludeTool/Program.cs
--- a/Engine/Source/Programs/IncludeTool/IncludeTool/Program.cs
+++ b/Engine/Source/Programs/IncludeTool/IncludeTool/Program.cs
@@ -284,7 +284,7 @@ namespace IncludeTool
 				}
 				else
 				{
-					InputDir = new DirectoryReference(Path.Combine(Path.GetDirectoryName(Assembly.GetExecutingAssembly().Location), "..\\..\\.."));
+					InputDir = new DirectoryReference(Path.Combine(Path.GetDirectoryName(Assembly.GetExecutingAssembly().Location), "..\\..\\..\\.."));
 				}
 
 				// Generate the exported task list 
@@ -856,7 +856,7 @@ namespace IncludeTool
 		/// <summary>
 		/// Create a task list for compiling a given target
 		/// </summary>
-		/// <param name="RootDirectory">The root directory for the branch</param>
+		/// <param name="RootDir">The root directory for the branch</param>
 		/// <param name="Target">The target to compile (eg. UE4Editor Win64 Development)</param>
 		/// <param name="TaskListFile">The output file for the task list</param>
 		/// <param name="Log">Log to output messages to</param>
