Remove UDK Remote and Slate Remote

diff --git a/Engine/Plugins/Slate/SlateRemote/SlateRemote.uplugin b/Engine/Plugins/Slate/SlateRemote/SlateRemote.uplugin
deleted file mode 100644
diff --git a/Engine/Plugins/Slate/SlateRemote/Source/SlateRemote/Private/Server/SlateRemoteServer.cpp b/Engine/Plugins/Slate/SlateRemote/Source/SlateRemote/Private/Server/SlateRemoteServer.cpp
deleted file mode 100644
diff --git a/Engine/Plugins/Slate/SlateRemote/Source/SlateRemote/Private/Server/SlateRemoteServer.h b/Engine/Plugins/Slate/SlateRemote/Source/SlateRemote/Private/Server/SlateRemoteServer.h
deleted file mode 100644
diff --git a/Engine/Plugins/Slate/SlateRemote/Source/SlateRemote/Private/Server/SlateRemoteServerMessage.h b/Engine/Plugins/Slate/SlateRemote/Source/SlateRemote/Private/Server/SlateRemoteServerMessage.h
deleted file mode 100644
diff --git a/Engine/Plugins/Slate/SlateRemote/Source/SlateRemote/Private/Shared/SlateRemoteSettings.h b/Engine/Plugins/Slate/SlateRemote/Source/SlateRemote/Private/Shared/SlateRemoteSettings.h
deleted file mode 100644
diff --git a/Engine/Plugins/Slate/SlateRemote/Source/SlateRemote/Private/SlateRemoteClasses.cpp b/Engine/Plugins/Slate/SlateRemote/Source/SlateRemote/Private/SlateRemoteClasses.cpp
deleted file mode 100644
diff --git a/Engine/Plugins/Slate/SlateRemote/Source/SlateRemote/Private/SlateRemoteModule.cpp b/Engine/Plugins/Slate/SlateRemote/Source/SlateRemote/Private/SlateRemoteModule.cpp
deleted file mode 100644
diff --git a/Engine/Plugins/Slate/SlateRemote/Source/SlateRemote/Private/SlateRemotePrivate.h b/Engine/Plugins/Slate/SlateRemote/Source/SlateRemote/Private/SlateRemotePrivate.h
deleted file mode 100644
diff --git a/Engine/Plugins/Slate/SlateRemote/Source/SlateRemote/SlateRemote.Build.cs b/Engine/Plugins/Slate/SlateRemote/Source/SlateRemote/SlateRemote.Build.cs
deleted file mode 100644
diff --git a/Engine/Source/Programs/IOS/UDKRemote/BuildIcons.sh b/Engine/Source/Programs/IOS/UDKRemote/BuildIcons.sh
deleted file mode 100755
diff --git a/Engine/Source/Programs/IOS/UDKRemote/Classes/Engine.h b/Engine/Source/Programs/IOS/UDKRemote/Classes/Engine.h
deleted file mode 100644
diff --git a/Engine/Source/Programs/IOS/UDKRemote/Classes/FlipsideViewController.h b/Engine/Source/Programs/IOS/UDKRemote/Classes/FlipsideViewController.h
deleted file mode 100644
diff --git a/Engine/Source/Programs/IOS/UDKRemote/Classes/FlipsideViewController.mm b/Engine/Source/Programs/IOS/UDKRemote/Classes/FlipsideViewController.mm
deleted file mode 100644
diff --git a/Engine/Source/Programs/IOS/UDKRemote/Classes/IPhoneAsyncTask.h b/Engine/Source/Programs/IOS/UDKRemote/Classes/IPhoneAsyncTask.h
deleted file mode 100644
diff --git a/Engine/Source/Programs/IOS/UDKRemote/Classes/IPhoneAsyncTask.mm b/Engine/Source/Programs/IOS/UDKRemote/Classes/IPhoneAsyncTask.mm
deleted file mode 100644
diff --git a/Engine/Source/Programs/IOS/UDKRemote/Classes/IPhoneHome.h b/Engine/Source/Programs/IOS/UDKRemote/Classes/IPhoneHome.h
deleted file mode 100644
diff --git a/Engine/Source/Programs/IOS/UDKRemote/Classes/IPhoneHome.mm b/Engine/Source/Programs/IOS/UDKRemote/Classes/IPhoneHome.mm
deleted file mode 100644
diff --git a/Engine/Source/Programs/IOS/UDKRemote/Classes/IPhoneObjCWrapper.h b/Engine/Source/Programs/IOS/UDKRemote/Classes/IPhoneObjCWrapper.h
deleted file mode 100644
diff --git a/Engine/Source/Programs/IOS/UDKRemote/Classes/IPhoneObjCWrapper.mm b/Engine/Source/Programs/IOS/UDKRemote/Classes/IPhoneObjCWrapper.mm
deleted file mode 100644
diff --git a/Engine/Source/Programs/IOS/UDKRemote/Classes/MainViewController.h b/Engine/Source/Programs/IOS/UDKRemote/Classes/MainViewController.h
deleted file mode 100644
diff --git a/Engine/Source/Programs/IOS/UDKRemote/Classes/MainViewController.mm b/Engine/Source/Programs/IOS/UDKRemote/Classes/MainViewController.mm
deleted file mode 100644
diff --git a/Engine/Source/Programs/IOS/UDKRemote/Classes/TextCell.h b/Engine/Source/Programs/IOS/UDKRemote/Classes/TextCell.h
deleted file mode 100644
diff --git a/Engine/Source/Programs/IOS/UDKRemote/Classes/TextCell.m b/Engine/Source/Programs/IOS/UDKRemote/Classes/TextCell.m
deleted file mode 100644
diff --git a/Engine/Source/Programs/IOS/UDKRemote/Classes/UDKRemoteAppDelegate.h b/Engine/Source/Programs/IOS/UDKRemote/Classes/UDKRemoteAppDelegate.h
deleted file mode 100644
diff --git a/Engine/Source/Programs/IOS/UDKRemote/Classes/UDKRemoteAppDelegate.mm b/Engine/Source/Programs/IOS/UDKRemote/Classes/UDKRemoteAppDelegate.mm
deleted file mode 100644
diff --git a/Engine/Source/Programs/IOS/UDKRemote/UDKRemote/Images.xcassets/AppIcon.appiconset/Contents.json b/Engine/Source/Programs/IOS/UDKRemote/UDKRemote/Images.xcassets/AppIcon.appiconset/Contents.json
deleted file mode 100644
diff --git a/Engine/Source/Programs/IOS/UDKRemote/UDKRemote/Images.xcassets/Contents.json b/Engine/Source/Programs/IOS/UDKRemote/UDKRemote/Images.xcassets/Contents.json
deleted file mode 100644
diff --git a/Engine/Source/Programs/IOS/UDKRemote/UDKRemote/Images.xcassets/LaunchImage.launchimage/Contents.json b/Engine/Source/Programs/IOS/UDKRemote/UDKRemote/Images.xcassets/LaunchImage.launchimage/Contents.json
deleted file mode 100644
diff --git a/Engine/Source/Programs/IOS/UDKRemote/main.mm b/Engine/Source/Programs/IOS/UDKRemote/main.mm
deleted file mode 100644
diff --git a/Engine/Source/Programs/IncludeTool/IncludeTool/Rules.cs b/Engine/Source/Programs/IncludeTool/IncludeTool/Rules.cs
--- a/Engine/Source/Programs/IncludeTool/IncludeTool/Rules.cs
+++ b/Engine/Source/Programs/IncludeTool/IncludeTool/Rules.cs
@@ -27,7 +27,6 @@ namespace IncludeTool
 			"/engine/source/thirdparty/mcpp/mcpp-2.7.2/test-c/",
 			"/engine/source/programs/unrealswarm/private/",
 			"/engine/plugins/runtime/packethandlers/compressioncomponents/oodle/source/thirdparty/notforlicensees/oodle/213/win/examples/",
-			"/engine/source/programs/ios/udkremote/"
 		};
 
 		/// <summary>
