Bink: Maybe fix for lint issue with includes?

diff --git a/Engine/Source/Programs/IncludeTool/IncludeTool/Rules.cs b/Engine/Source/Programs/IncludeTool/IncludeTool/Rules.cs
--- a/Engine/Source/Programs/IncludeTool/IncludeTool/Rules.cs
+++ b/Engine/Source/Programs/IncludeTool/IncludeTool/Rules.cs
@@ -65,6 +65,8 @@ namespace IncludeTool
 		/// </summary>
 		static readonly string[] ExternalFileIncludePaths =
 		{
+			"binkplugin.h",
+			"egttypes.h",
 			"oodle.h",
 			"oodle2.h",
 			"xg.h",
