Bink: Fix for include after extern

diff --git a/Engine/Source/Programs/IncludeTool/IncludeTool/Rules.cs b/Engine/Source/Programs/IncludeTool/IncludeTool/Rules.cs
--- a/Engine/Source/Programs/IncludeTool/IncludeTool/Rules.cs
+++ b/Engine/Source/Programs/IncludeTool/IncludeTool/Rules.cs
@@ -65,6 +65,7 @@ namespace IncludeTool
 		/// </summary>
 		static readonly string[] ExternalFileIncludePaths =
 		{
+			"BinkMediaPlayerPCH.h",
 			"binkplugin.h",
 			"egttypes.h",
 			"oodle.h",
