Bink: better fix for last lint issue submit. This way I don't have to exclude that file.

diff --git a/Engine/Source/Programs/IncludeTool/IncludeTool/Rules.cs b/Engine/Source/Programs/IncludeTool/IncludeTool/Rules.cs
--- a/Engine/Source/Programs/IncludeTool/IncludeTool/Rules.cs
+++ b/Engine/Source/Programs/IncludeTool/IncludeTool/Rules.cs
@@ -65,7 +65,6 @@ namespace IncludeTool
 		/// </summary>
 		static readonly string[] ExternalFileIncludePaths =
 		{
-			"BinkMediaPlayerPCH.h",
 			"binkplugin.h",
 			"egttypes.h",
 			"oodle.h",
