IncludeTool: Check for forward slashes in include directives as they cause problems for native linux

#rnx

diff --git a/Engine/Source/Programs/IncludeTool/IncludeTool/Program.cs b/Engine/Source/Programs/IncludeTool/IncludeTool/Program.cs
--- a/Engine/Source/Programs/IncludeTool/IncludeTool/Program.cs
+++ b/Engine/Source/Programs/IncludeTool/IncludeTool/Program.cs
@@ -623,6 +623,19 @@ namespace IncludeTool
 						}
 					}
 
+					// Make sure include directives only use forward slashes
+					foreach (SourceFile PreprocessedFile in PreprocessedFiles)
+					{
+						PreprocessorMarkup[] Results = Array.FindAll(PreprocessedFile.Markup, x => x.Type == PreprocessorMarkupType.Include && x.IsActive);
+						foreach (PreprocessorMarkup Result in Results)
+						{
+							if (Result.ToString().Contains("\\"))
+							{
+								Log.WriteLine("warning: include directive using '\\' as path separator {0}{1}", PreprocessedFile.Location, Result);
+							}
+						}
+					}
+
 					// Print all the conflicting declarations
 					SymbolTable.PrintConflicts(Log);
 				}
