NeuralNetworkInference (NNI) plugin:
- Removed static warning for NNI.

#rnx

diff --git a/Engine/Source/Programs/IncludeTool/IncludeTool/Program.cs b/Engine/Source/Programs/IncludeTool/IncludeTool/Program.cs
--- a/Engine/Source/Programs/IncludeTool/IncludeTool/Program.cs
+++ b/Engine/Source/Programs/IncludeTool/IncludeTool/Program.cs
@@ -1039,7 +1039,14 @@ namespace IncludeTool
 				{
 					foreach(PreprocessorMarkup Markup in File.Markup)
 					{
-						if(Markup.IncludedFile != null && (Markup.IncludedFile.Flags & SourceFileFlags.External) == 0 && !File.Location.FullName.Contains("PhyaLib") && !File.Location.GetFileName().Contains("Recast") && !File.Location.GetFileName().Contains("FramePro") && !File.Location.FullName.Contains("libSampleRate") && !File.Location.FullName.Contains("lz4"))
+						if (Markup.IncludedFile != null && (Markup.IncludedFile.Flags & SourceFileFlags.External) == 0
+							&& !File.Location.GetFileName().Contains("FramePro")
+							&& !File.Location.GetFileName().Contains("Recast")
+							&& !File.Location.FullName.Contains("libSampleRate")
+							&& !File.Location.FullName.Contains("lz4")
+							&& !File.Location.FullName.Contains("NeuralNetworkInference")
+							&& !File.Location.FullName.Contains("PhyaLib")
+							)
 						{
 							Log.WriteLine("{0}({1}): warning: external file is including non-external file ('{2}')", IncludedFile.Location.FullName, Markup.Location.LineIdx + 1, Markup.IncludedFile.Location);
 						}
