IncludeTool: Output messages in a format that can be parsed for warnings/errors by Horde.

diff --git a/Engine/Source/Programs/IncludeTool/IncludeTool/Program.cs b/Engine/Source/Programs/IncludeTool/IncludeTool/Program.cs
--- a/Engine/Source/Programs/IncludeTool/IncludeTool/Program.cs
+++ b/Engine/Source/Programs/IncludeTool/IncludeTool/Program.cs
@@ -565,7 +565,7 @@ namespace IncludeTool
 					{
 						if(PreprocessedFile.Location.HasExtension("PCH.h") && (PreprocessedFile.Flags & SourceFileFlags.Aggregate) == 0)
 						{
-							Log.WriteLine("warning: File appears to be PCH but also contains declarations. Will not be optimized out: {0}", PreprocessedFile.Location);
+							Log.WriteWarning(PreprocessedFile.Location, "file appears to be PCH but also contains declarations, will not be optimized out.");
 						}
 					}
 
@@ -574,7 +574,7 @@ namespace IncludeTool
 					{
 						if(PreprocessedFile.HasHeaderGuard && PreprocessedFile.BodyMaxIdx < PreprocessedFile.Markup.Length && (PreprocessedFile.Flags & SourceFileFlags.External) == 0 && Rules.IgnoreOldStyleHeaderGuards("/" + PreprocessedFile.Location.MakeRelativeTo(InputDir).ToLowerInvariant()))
 						{
-							Log.WriteLine("warning: file has old-style header guard: {0}", PreprocessedFile.Location.FullName);
+							Log.WriteWarning(PreprocessedFile.Location, "file has old-style header guard");
 						}
 					}
 
@@ -585,7 +585,7 @@ namespace IncludeTool
 						{
 							if(!PreprocessedFile.Location.GetFileName().Equals("MonolithicHeaderBoilerplate.h", StringComparison.OrdinalIgnoreCase))
 							{
-								Log.WriteLine("warning: missing header guard: {0}", PreprocessedFile.Location.FullName);
+								Log.WriteWarning(PreprocessedFile.Location, "missing header guard");
 							}
 						}
 					}
@@ -601,7 +601,7 @@ namespace IncludeTool
 								int LastIncludeIdx = Array.FindLastIndex(PreprocessedFile.Markup, x => x.Type == PreprocessorMarkupType.Include && x.IsActive && (x.IncludedFile.Flags & (SourceFileFlags.Pinned | SourceFileFlags.Inline)) == 0);
 								if(FirstTextIdx < LastIncludeIdx)
 								{
-									Log.WriteLine("warning: includes after first code block: {0}", PreprocessedFile.Location);
+									Log.WriteWarning(PreprocessedFile.Location, "includes after first code block");
 								}
 							}
 						}
@@ -614,11 +614,11 @@ namespace IncludeTool
 						{
 							if(PreprocessedFile.Fragments.Length != 1)
 							{
-								Log.WriteLine("warning: expected only one fragment in forward declaration header: {0}", PreprocessedFile.Location);
+								Log.WriteWarning(PreprocessedFile.Location, "expected only one fragment in forward declaration header");
 							}
 							else if(PreprocessedFile.Markup.Skip(PreprocessedFile.BodyMinIdx).Any(x => x.Type != PreprocessorMarkupType.Include && x.Type != PreprocessorMarkupType.Text))
 							{
-								Log.WriteLine("warning: expected only include directives and text in forward declaration header: {0}", PreprocessedFile.Location);
+								Log.WriteWarning(PreprocessedFile.Location, "expected only include directives and text in forward declaration header");
 							}
 						}
 					}
@@ -631,7 +631,7 @@ namespace IncludeTool
 						{
 							if (Result.ToString().Contains("\\"))
 							{
-								Log.WriteLine("warning: include directive using '\\' as path separator {0}{1}", PreprocessedFile.Location, Result);
+								Log.WriteWarning(PreprocessedFile.Location, "include directive using '\\' as path separator {0}", Result);
 							}
 						}
 					}
diff --git a/Engine/Source/Programs/IncludeTool/IncludeTool/Support/LineBasedTextWriter.cs b/Engine/Source/Programs/IncludeTool/IncludeTool/Support/LineBasedTextWriter.cs
--- a/Engine/Source/Programs/IncludeTool/IncludeTool/Support/LineBasedTextWriter.cs
+++ b/Engine/Source/Programs/IncludeTool/IncludeTool/Support/LineBasedTextWriter.cs
@@ -24,5 +24,15 @@ namespace IncludeTool.Support
 		{
 			WriteLine(String.Format(Format, Args));
 		}
+
+		public void WriteWarning(FileReference File, string Message)
+		{
+			WriteLine("{0}: warning: {1}", File, Message);
+		}
+
+		public void WriteWarning(FileReference File, string Message, params object[] Args)
+		{
+			WriteWarning(File, String.Format(Message, Args));
+		}
 	}
 }
