Fix static analysis
#rnx

diff --git a/Engine/Source/Programs/IncludeTool/IncludeTool/Rules.cs b/Engine/Source/Programs/IncludeTool/IncludeTool/Rules.cs
--- a/Engine/Source/Programs/IncludeTool/IncludeTool/Rules.cs
+++ b/Engine/Source/Programs/IncludeTool/IncludeTool/Rules.cs
@@ -428,7 +428,7 @@ namespace IncludeTool
 			"/Engine/Source/Runtime/Experimental/Chaos/Public/Chaos/PBDRigidsEvolutionFwd.h", // invalid forward declaration - 'namespace Chaos'
 			"/Engine/Source/Runtime/Experimental/Chaos/Public/PhysicsProxy/JointConstraintProxyFwd.h", // invalid forward declaration - 'namespace Chaos'
 			"/Engine/Source/Runtime/Experimental/Chaos/Public/PhysicsProxy/SingleParticlePhysicsProxyFwd.h", // invalid forward declaration - 'namespace Chaos'
-			"/Engine/Source/Runtime/Interchange/Engine/Public/InterchangeEngineFwd.h", // invalid forward declaration - 'namespace UE'
+			"/Engine/Source/Runtime/Experimental/Interchange/Engine/Public/InterchangeEngineFwd.h", // invalid forward declaration - 'namespace UE'
 		};
 
 		/// <summary>
