IncludeTool: Fix exception skipping over lines containing raw string literals with embedded newlines.

diff --git a/Engine/Source/Programs/IncludeTool/IncludeTool/TokenReader.cs b/Engine/Source/Programs/IncludeTool/IncludeTool/TokenReader.cs
--- a/Engine/Source/Programs/IncludeTool/IncludeTool/TokenReader.cs
+++ b/Engine/Source/Programs/IncludeTool/IncludeTool/TokenReader.cs
@@ -214,23 +214,17 @@ namespace IncludeTool
 		/// <returns>True if we were able to find another token</returns>
 		public bool MoveToNextLine()
 		{
-			int NextLineIdx = TokenLocation.LineIdx;
-
-			// Skip past all the escaped newlines
-			while(NextLineIdx < Text.Lines.Length && Text.Lines[NextLineIdx].EndsWith("\\"))
+			for (; ; )
 			{
-				NextLineIdx++;
-			}
-
-			// And move to the next line
-			if(NextLineIdx < Text.Lines.Length)
-			{
-				NextLineIdx++;
+				if (Current.Text == "\n")
+				{
+					return MoveNext();
+				}
+				if (!MoveNext())
+				{
+					return false;
+				}
 			}
-
-			// Read a token
-			TokenEndLocation = new TextLocation(NextLineIdx, 0);
-			return MoveNext();
 		}
 
 		/// <summary>
