IncludeTool: Added an exclusion to avoid an error about a valid forward declaration

#rnx

diff --git a/Engine/Source/Programs/IncludeTool/IncludeTool/Rules.cs b/Engine/Source/Programs/IncludeTool/IncludeTool/Rules.cs
--- a/Engine/Source/Programs/IncludeTool/IncludeTool/Rules.cs
+++ b/Engine/Source/Programs/IncludeTool/IncludeTool/Rules.cs
@@ -411,6 +411,7 @@ namespace IncludeTool
 		/// </summary>
 		static readonly HashSet<string> IgnoreFwdHeaders = new HashSet<string>(StringComparer.OrdinalIgnoreCase)
 		{
+			"/Engine/Source/Developer/DerivedDataCache/Public/DerivedDataSharedStringFwd.h", // invalid forward declaration - 'namespace UE::DerivedData'
 			"/Engine/Source/Runtime/Core/Public/Internationalization/TextNamespaceFwd.h",
 			"/Engine/Source/Editor/SceneOutliner/Public/SceneOutlinerFwd.h",
 			"/Engine/Source/Editor/EditorFramework/Public/UnrealWidgetFwd.h", // error: invalid forward declaration - 'enum ECoordSystem'
