IncludeTool: Fix buffer overrun with malformed include directives.

diff --git a/Engine/Source/Programs/IncludeTool/IncludeTool/TokenReader.cs b/Engine/Source/Programs/IncludeTool/IncludeTool/TokenReader.cs
--- a/Engine/Source/Programs/IncludeTool/IncludeTool/TokenReader.cs
+++ b/Engine/Source/Programs/IncludeTool/IncludeTool/TokenReader.cs
@@ -430,7 +430,7 @@ namespace IncludeTool
 			else if(Character == '<' && Context.HasFlag(TokenReaderContext.IncludeDirective))
 			{
 				StringBuilder Builder = new StringBuilder("<");
-				while(Builder[Builder.Length - 1] != '>')
+				while(Builder[Builder.Length - 1] != '>' && ColumnIdx < Text.Lines[LineIdx].Length)
 				{
 					Builder.Append(Text[LineIdx, ColumnIdx]);
 					ColumnIdx++;
