IncludeTool: Line numbers for some warnings, fix check for old-style include guards, add exclusions

#rnx

diff --git a/Engine/Source/Programs/IncludeTool/IncludeTool/Program.cs b/Engine/Source/Programs/IncludeTool/IncludeTool/Program.cs
--- a/Engine/Source/Programs/IncludeTool/IncludeTool/Program.cs
+++ b/Engine/Source/Programs/IncludeTool/IncludeTool/Program.cs
@@ -573,7 +573,8 @@ namespace IncludeTool
 					// Warn about any files that appear to be PCHs but actually include declarations
 					foreach (SourceFile PreprocessedFile in PreprocessedFiles)
 					{
-						if(PreprocessedFile.Location.HasExtension("PCH.h") && (PreprocessedFile.Flags & SourceFileFlags.Aggregate) == 0)
+						if(PreprocessedFile.Location.HasExtension("PCH.h") && 
+						   (PreprocessedFile.Flags & SourceFileFlags.Aggregate) == 0)
 						{
 							Log.WriteWarning(PreprocessedFile.Location, "file appears to be PCH but also contains declarations, will not be optimized out.");
 						}
@@ -582,28 +583,34 @@ namespace IncludeTool
 					// Test for any headers with old-style include guards
 					foreach(SourceFile PreprocessedFile in PreprocessedFiles)
 					{
-						if(PreprocessedFile.HasHeaderGuard && PreprocessedFile.BodyMaxIdx < PreprocessedFile.Markup.Length && (PreprocessedFile.Flags & SourceFileFlags.External) == 0 && Rules.IgnoreOldStyleHeaderGuards("/" + PreprocessedFile.Location.MakeRelativeTo(InputDir).ToLowerInvariant()))
+						if(PreprocessedFile.HasHeaderGuard && 
+						   PreprocessedFile.BodyMaxIdx < PreprocessedFile.Markup.Length && 
+						   (PreprocessedFile.Flags & SourceFileFlags.External) == 0 && 
+						   !Rules.IgnoreOldStyleHeaderGuards(PreprocessedFile.Location))
 						{
-							Log.WriteWarning(PreprocessedFile.Location, "file has old-style header guard");
+							int LineIdx = PreprocessedFile.Markup[0].Location.LineIdx;
+							Log.WriteWarning(PreprocessedFile.Location, LineIdx, $"File has old-style header guard: {PreprocessedFile.Text.Lines[LineIdx]}");
 						}
 					}
 
 					// Test for any headers without include guards
 					foreach(SourceFile PreprocessedFile in PreprocessedFiles)
 					{
-						if(!PreprocessedFile.HasHeaderGuard && (PreprocessedFile.Flags & (SourceFileFlags.TranslationUnit | SourceFileFlags.Inline | SourceFileFlags.External | SourceFileFlags.GeneratedHeader)) == 0 && PreprocessedFile.Counterpart == null && PreprocessedFile.Location.HasExtension(".h"))
+						if(!PreprocessedFile.HasHeaderGuard && 
+						   (PreprocessedFile.Flags & (SourceFileFlags.TranslationUnit | SourceFileFlags.Inline | SourceFileFlags.External | SourceFileFlags.GeneratedHeader)) == 0 && 
+						   PreprocessedFile.Counterpart == null && 
+						   PreprocessedFile.Location.HasExtension(".h") &&
+						   !PreprocessedFile.Location.GetFileName().Equals("MonolithicHeaderBoilerplate.h", StringComparison.OrdinalIgnoreCase))
 						{
-							if(!PreprocessedFile.Location.GetFileName().Equals("MonolithicHeaderBoilerplate.h", StringComparison.OrdinalIgnoreCase))
-							{
-								Log.WriteWarning(PreprocessedFile.Location, "missing header guard");
-							}
+							Log.WriteWarning(PreprocessedFile.Location, "missing header guard");
 						}
 					}
 
 					// Test for any headers that include files after the first block of code
 					foreach(SourceFile PreprocessedFile in PreprocessedFiles)
 					{
-						if((PreprocessedFile.Flags & SourceFileFlags.TranslationUnit) == 0 && (PreprocessedFile.Flags & SourceFileFlags.AllowMultipleFragments) == 0)
+						if((PreprocessedFile.Flags & SourceFileFlags.TranslationUnit) == 0 && 
+						   (PreprocessedFile.Flags & SourceFileFlags.AllowMultipleFragments) == 0)
 						{
 							int FirstTextIdx = Array.FindIndex(PreprocessedFile.Markup, x => x.Type == PreprocessorMarkupType.Text);
 							if(FirstTextIdx != -1)
@@ -611,7 +618,8 @@ namespace IncludeTool
 								int LastIncludeIdx = Array.FindLastIndex(PreprocessedFile.Markup, x => x.Type == PreprocessorMarkupType.Include && x.IsActive && (x.IncludedFile.Flags & (SourceFileFlags.Pinned | SourceFileFlags.Inline)) == 0);
 								if(FirstTextIdx < LastIncludeIdx)
 								{
-									Log.WriteWarning(PreprocessedFile.Location, "includes after first code block");
+									int LineIdx = PreprocessedFile.Markup[LastIncludeIdx].Location.LineIdx;
+									Log.WriteWarning(PreprocessedFile.Location, LineIdx, $"Include after first code block: {PreprocessedFile.Text.Lines[LineIdx]}");
 								}
 							}
 						}
@@ -641,7 +649,7 @@ namespace IncludeTool
 						{
 							if (Result.ToString().Contains("\\"))
 							{
-								Log.WriteWarning(PreprocessedFile.Location, "include directive using '\\' as path separator {0}", Result);
+								Log.WriteWarning(PreprocessedFile.Location, Result.Location.LineIdx, $"Include directive using '\\' as path separator: {PreprocessedFile.Text.Lines[Result.Location.LineIdx]}");
 							}
 						}
 					}
diff --git a/Engine/Source/Programs/IncludeTool/IncludeTool/Rules.cs b/Engine/Source/Programs/IncludeTool/IncludeTool/Rules.cs
--- a/Engine/Source/Programs/IncludeTool/IncludeTool/Rules.cs
+++ b/Engine/Source/Programs/IncludeTool/IncludeTool/Rules.cs
@@ -578,31 +578,27 @@ namespace IncludeTool
 			return false;
 		}
 
+		static readonly string[] PathsToIgnoreForOldStyleHeaders = new string[] {
+			"/Engine/Source/Runtime/Navmesh/Public/Detour/",
+			"/Engine/Source/Runtime/Navmesh/Public/DetourCrowd/",
+			"/Engine/Source/Runtime/Navmesh/Public/Recast/",
+			"/Engine/Source/Runtime/Navmesh/Public/DebugUtils/",
+			"/Engine/Source/Runtime/OodleDataCompression/Sdks/",
+
+			"/Engine/Source/Runtime/Core/Public/Experimental/Containers/FAAArrayQueue.h",
+			"/Engine/Source/Runtime/Core/Public/Hash/CityHash.h",
+			"/Engine/Source/Runtime/Core/Public/MemPro/MemPro.h",
+		};
+
 		/// <summary>
 		/// Whether to ignore old-style header guards
 		/// </summary>
-		/// <param name="RootDir">Root directory for the branch</param>
-		/// <param name="Location">Path to the file</param>
+		/// <param name="NormalizedPath">Path to the file</param>
 		/// <returns>True to ignore the different derivations</returns>
-		public static bool IgnoreOldStyleHeaderGuards(string NormalizedPath)
+		public static bool IgnoreOldStyleHeaderGuards(FileReference Location)
 		{
-			if(NormalizedPath.StartsWith("/engine/source/runtime/navmesh/public/detour/"))
-			{
-				return true;
-			}
-			if(NormalizedPath.StartsWith("/engine/source/runtime/navmesh/public/detourcrowd/"))
-			{
-				return true;
-			}
-			if(NormalizedPath.StartsWith("/engine/source/runtime/navmesh/public/recast/"))
-			{
-				return true;
-			}
-			if(NormalizedPath.StartsWith("/engine/source/runtime/navmesh/public/debugutils/"))
-			{
-				return true;
-			}
-			return false;
+			string LocationString = Location.ToString().Replace("\\", "/");
+			return PathsToIgnoreForOldStyleHeaders.Any(x => LocationString.Contains(x));
 		}
 
         /// <summary>
diff --git a/Engine/Source/Programs/IncludeTool/IncludeTool/Support/LineBasedTextWriter.cs b/Engine/Source/Programs/IncludeTool/IncludeTool/Support/LineBasedTextWriter.cs
--- a/Engine/Source/Programs/IncludeTool/IncludeTool/Support/LineBasedTextWriter.cs
+++ b/Engine/Source/Programs/IncludeTool/IncludeTool/Support/LineBasedTextWriter.cs
@@ -26,13 +26,15 @@ namespace IncludeTool.Support
 		}
 
 		public void WriteWarning(FileReference File, string Message)
-		{
-			WriteLine("{0}: warning: {1}", File, Message);
-		}
+			=> WriteLine($"{File}: warning: {Message}");
+
+		public void WriteWarning(FileReference File, int LineNumber, string Message)
+			=> WriteLine($"{File}({LineNumber + 1}): warning: {Message}");
 
 		public void WriteWarning(FileReference File, string Message, params object[] Args)
-		{
-			WriteWarning(File, String.Format(Message, Args));
-		}
+			=> WriteWarning(File, String.Format(Message, Args));
+
+		public void WriteWarning(FileReference File, int LineNumber, string Message, params object[] Args)
+			=> WriteWarning(File, LineNumber, String.Format(Message, Args));
 	}
 }
