IncludeTool: More exclusions for old-style header guards

#trivial
#rnx

diff --git a/Engine/Source/Programs/IncludeTool/IncludeTool/Rules.cs b/Engine/Source/Programs/IncludeTool/IncludeTool/Rules.cs
--- a/Engine/Source/Programs/IncludeTool/IncludeTool/Rules.cs
+++ b/Engine/Source/Programs/IncludeTool/IncludeTool/Rules.cs
@@ -579,15 +579,28 @@ namespace IncludeTool
 		}
 
 		static readonly string[] PathsToIgnoreForOldStyleHeaders = new string[] {
+			"/Engine/Source/Runtime/Navmesh/Public/DebugUtils/",
 			"/Engine/Source/Runtime/Navmesh/Public/Detour/",
 			"/Engine/Source/Runtime/Navmesh/Public/DetourCrowd/",
+			"/Engine/Source/Runtime/Navmesh/Public/DetourTileCache/",
 			"/Engine/Source/Runtime/Navmesh/Public/Recast/",
-			"/Engine/Source/Runtime/Navmesh/Public/DebugUtils/",
+
+			"/Engine/Plugins/Compression/OodleNetwork/",
+			"/Engine/Plugins/Developer/TextureFormatOodle/",
+			"/Engine/Plugins/Media/BinkMedia/Source/",
 			"/Engine/Source/Runtime/OodleDataCompression/Sdks/",
 
+			"/Engine/Plugins/Animation/ControlRigSpline/Source/ControlRigSpline/ThirdParty/",
+			"/Engine/Plugins/Runtime/nDisplay/ThirdParty/",
+			"/Engine/Plugins/Runtime/ResonanceAudio/Source/ResonanceAudio/Private/ResonanceAudioLibrary/resonance_audio/",
+			"/Engine/Source/Runtime/Experimental/Voronoi/",
+
+			"/Engine/Plugins/Animation/RigLogic/Source/RigLogicLib/Private/dna/utils/Extd.h",
 			"/Engine/Source/Runtime/Core/Public/Experimental/Containers/FAAArrayQueue.h",
 			"/Engine/Source/Runtime/Core/Public/Hash/CityHash.h",
 			"/Engine/Source/Runtime/Core/Public/MemPro/MemPro.h",
+
+			"/Engine/Restricted/",
 		};
 
 		/// <summary>
