IncludeTool: More exclusion for old-style header guards

#trivial
#rnx

diff --git a/Engine/Source/Programs/IncludeTool/IncludeTool/Rules.cs b/Engine/Source/Programs/IncludeTool/IncludeTool/Rules.cs
--- a/Engine/Source/Programs/IncludeTool/IncludeTool/Rules.cs
+++ b/Engine/Source/Programs/IncludeTool/IncludeTool/Rules.cs
@@ -592,10 +592,12 @@ namespace IncludeTool
 
 			"/Engine/Plugins/Animation/ControlRigSpline/Source/ControlRigSpline/ThirdParty/",
 			"/Engine/Plugins/Runtime/nDisplay/ThirdParty/",
-			"/Engine/Plugins/Runtime/ResonanceAudio/Source/ResonanceAudio/Private/ResonanceAudioLibrary/resonance_audio/",
+			"/Engine/Plugins/Runtime/ResonanceAudio/Source/ResonanceAudio/Private/ResonanceAudioLibrary/",
 			"/Engine/Source/Runtime/Experimental/Voronoi/",
 
 			"/Engine/Plugins/Animation/RigLogic/Source/RigLogicLib/Private/dna/utils/Extd.h",
+			"/Engine/Plugins/Animation/RigLogic/Source/RigLogicLib/Private/riglogic/utils/Extd.h",
+			"/Engine/Plugins/Animation/RigLogic/Source/RigLogicLib/Public/dna/types/ArrayView.h",
 			"/Engine/Source/Runtime/Core/Public/Experimental/Containers/FAAArrayQueue.h",
 			"/Engine/Source/Runtime/Core/Public/Hash/CityHash.h",
 			"/Engine/Source/Runtime/Core/Public/MemPro/MemPro.h",
