IncludeTool: even more exclusions for old-style header guards

#trivial
#rnx

diff --git a/Engine/Source/Programs/IncludeTool/IncludeTool/Rules.cs b/Engine/Source/Programs/IncludeTool/IncludeTool/Rules.cs
--- a/Engine/Source/Programs/IncludeTool/IncludeTool/Rules.cs
+++ b/Engine/Source/Programs/IncludeTool/IncludeTool/Rules.cs
@@ -595,12 +595,15 @@ namespace IncludeTool
 			"/Engine/Plugins/Runtime/ResonanceAudio/Source/ResonanceAudio/Private/ResonanceAudioLibrary/",
 			"/Engine/Source/Runtime/Experimental/Voronoi/",
 
+			"/Engine/Plugins/Animation/ControlRig/Source/ControlRig/ThirdParty/AHEasing/AHEasing/easing.h",
 			"/Engine/Plugins/Animation/RigLogic/Source/RigLogicLib/Private/dna/utils/Extd.h",
 			"/Engine/Plugins/Animation/RigLogic/Source/RigLogicLib/Private/riglogic/utils/Extd.h",
 			"/Engine/Plugins/Animation/RigLogic/Source/RigLogicLib/Public/dna/types/ArrayView.h",
 			"/Engine/Source/Runtime/Core/Public/Experimental/Containers/FAAArrayQueue.h",
 			"/Engine/Source/Runtime/Core/Public/Hash/CityHash.h",
 			"/Engine/Source/Runtime/Core/Public/MemPro/MemPro.h",
+			"/Engine/Source/Runtime/CUDA/Source/Public/CudaWrapper.h",
+			"/Engine/Plugins/Runtime/GeometryProcessing/Source/GeometryAlgorithms/Private/ThirdParty/xatlas/xatlas.h",
 
 			"/Engine/Restricted/",
 		};
