Renamed remaining folders from NeuralNetworkInference to NNI missed by CL 19242774.
Updated IncludeTool to silence warnings for new plugin name.

diff --git a/Engine/Source/Programs/IncludeTool/IncludeTool/Program.cs b/Engine/Source/Programs/IncludeTool/IncludeTool/Program.cs
--- a/Engine/Source/Programs/IncludeTool/IncludeTool/Program.cs
+++ b/Engine/Source/Programs/IncludeTool/IncludeTool/Program.cs
@@ -1063,6 +1063,7 @@ namespace IncludeTool
 							&& !File.Location.FullName.Contains("libSampleRate")
 							&& !File.Location.FullName.Contains("lz4")
 							&& !File.Location.FullName.Contains("NeuralNetworkInference")
+							&& !File.Location.FullName.Contains("NNI")
 							&& !File.Location.FullName.Contains("PhyaLib")
 							)
 						{
