IncludeTool: allow differences in UE_DEFINE_LEGACY_MATH_CONSTANT_MACRO_NAMES and UE_DEPRECATE_LEGACY_MATH_CONSTANT_MACRO_NAMES

#rnx

diff --git a/Engine/Source/Programs/IncludeTool/IncludeTool/Rules.cs b/Engine/Source/Programs/IncludeTool/IncludeTool/Rules.cs
--- a/Engine/Source/Programs/IncludeTool/IncludeTool/Rules.cs
+++ b/Engine/Source/Programs/IncludeTool/IncludeTool/Rules.cs
@@ -555,7 +555,12 @@ namespace IncludeTool
 			if(Markup.Type == PreprocessorMarkupType.Define && Markup.Tokens[0].Text == "ONLINE_LOG_PREFIX")
 			{
 				return true;
-			}			if(Markup.Type == PreprocessorMarkupType.Define && (Markup.Tokens[0].Text == "UE_DEPRECATED_FORGAME" || Markup.Tokens[0].Text == "DEPRECATED_FORGAME"))
+			}			
+			if(Markup.Type == PreprocessorMarkupType.Define && (Markup.Tokens[0].Text == "UE_DEPRECATED_FORGAME" || Markup.Tokens[0].Text == "DEPRECATED_FORGAME"))
+			{
+				return true;
+			}
+			if(Markup.Type == PreprocessorMarkupType.Define && (Markup.Tokens[0].Text == "UE_DEFINE_LEGACY_MATH_CONSTANT_MACRO_NAMES" || Markup.Tokens[0].Text == "UE_DEPRECATE_LEGACY_MATH_CONSTANT_MACRO_NAMES"))
 			{
 				return true;
 			}
