IncludeTool: add UE_INCLUDETOOL_IGNORE_INCONSISTENT_STATE as a way to turn off warnings about inconsistent state, from code, in IncludeTool

#define UE_INCLUDETOOL_IGNORE_INCONSISTENT_STATE

... warning is ignored

#undef UE_INCLUDETOOL_IGNORE_INCONSISTENT_STATE

#rnx

diff --git a/Engine/Source/Programs/IncludeTool/IncludeTool/Program.cs b/Engine/Source/Programs/IncludeTool/IncludeTool/Program.cs
--- a/Engine/Source/Programs/IncludeTool/IncludeTool/Program.cs
+++ b/Engine/Source/Programs/IncludeTool/IncludeTool/Program.cs
@@ -1089,9 +1089,29 @@ namespace IncludeTool
 					}
 					else if((File.Flags & SourceFileFlags.External) == 0)
 					{
+						bool bIgnoreInconsistentState = false;
+						int IgnoreInconsistentStateIdx = -1;
+
 						// Another thread has already derived the active markup for this file; check it's identical.
 						for(int Idx = 0; Idx < File.Markup.Length; Idx++)
 						{
+							if (File.Markup[Idx].Type == PreprocessorMarkupType.Define && File.Markup[Idx].Tokens[0].Text == "UE_INCLUDETOOL_IGNORE_INCONSISTENT_STATE")
+							{
+								IgnoreInconsistentStateIdx = Idx;
+								bIgnoreInconsistentState = true;
+							}
+
+							if (bIgnoreInconsistentState)
+							{
+								if (File.Markup[Idx].Type == PreprocessorMarkupType.Undef && File.Markup[Idx].Tokens[0].Text == "UE_INCLUDETOOL_IGNORE_INCONSISTENT_STATE")
+								{
+									IgnoreInconsistentStateIdx = -1;
+									bIgnoreInconsistentState = false;
+								}
+
+								continue;
+							}
+
 							if(ActiveMarkup.Flags[Idx] != CurrentActiveMarkup.Flags[Idx] && !Rules.AllowDifferentMarkup(File, Idx))
 							{
 								lock(Log)
@@ -1110,6 +1130,11 @@ namespace IncludeTool
 								}
 							}
 						}
+
+						if (bIgnoreInconsistentState)
+						{
+							Log.WriteLine($"{File.Location.FullName}({File.Markup[IgnoreInconsistentStateIdx].Location.LineIdx + 1}): warning: UE_INCLUDE_IGNORE_INCONSISTENT_STATE was defined, but not undefined. Add '#undef UE_INCLUDE_IGNORE_INCONSISTENT_STATE' to this file.");
+						}
 					}
 				}
 
diff --git a/Engine/Source/Programs/IncludeTool/IncludeTool/Rules.cs b/Engine/Source/Programs/IncludeTool/IncludeTool/Rules.cs
--- a/Engine/Source/Programs/IncludeTool/IncludeTool/Rules.cs
+++ b/Engine/Source/Programs/IncludeTool/IncludeTool/Rules.cs
@@ -560,10 +560,6 @@ namespace IncludeTool
 			{
 				return true;
 			}
-			if(Markup.Type == PreprocessorMarkupType.Define && (Markup.Tokens[0].Text == "UE_DEFINE_LEGACY_MATH_CONSTANT_MACRO_NAMES" || Markup.Tokens[0].Text == "UE_DEPRECATE_LEGACY_MATH_CONSTANT_MACRO_NAMES"))
-			{
-				return true;
-			}
 			if(Markup.Type == PreprocessorMarkupType.Undef && Markup.Tokens[0].Text == "TEXT")
 			{
 				return true;
diff --git a/Engine/Source/Runtime/Core/Core.Build.cs b/Engine/Source/Runtime/Core/Core.Build.cs
--- a/Engine/Source/Runtime/Core/Core.Build.cs
+++ b/Engine/Source/Runtime/Core/Core.Build.cs
@@ -313,6 +313,8 @@ public class Core : ModuleRules
 		PrivateDefinitions.Add("PLATFORM_COMPILER_OPTIMIZATION_PG=" + (Target.bPGOOptimize ? "1" : "0"));
 		PrivateDefinitions.Add("PLATFORM_COMPILER_OPTIMIZATION_PG_PROFILING=" + (Target.bPGOProfile ? "1" : "0"));
 
+		PrivateDefinitions.Add("UE_DEFINE_LEGACY_MATH_CONSTANT_MACRO_NAMES=0");
+
 		UnsafeTypeCastWarningLevel = WarningLevel.Warning;
 	}
 
diff --git a/Engine/Source/Runtime/Core/Public/Math/UnrealMathUtility.h b/Engine/Source/Runtime/Core/Public/Math/UnrealMathUtility.h
--- a/Engine/Source/Runtime/Core/Public/Math/UnrealMathUtility.h
+++ b/Engine/Source/Runtime/Core/Public/Math/UnrealMathUtility.h
@@ -43,6 +43,10 @@ class TRange;
 	Floating point constants.
 -----------------------------------------------------------------------------*/
 
+// These macros can have different values across modules inside a single codebase.
+// Tell the IncludeTool static analyzer to ignore that divergence in state:
+#define UE_INCLUDETOOL_IGNORE_INCONSISTENT_STATE
+
 // Define this to 1 in a module's .Build.cs to make legacy names issue deprecation warnings.
 #ifndef UE_DEPRECATE_LEGACY_MATH_CONSTANT_MACRO_NAMES
 	#define UE_DEPRECATE_LEGACY_MATH_CONSTANT_MACRO_NAMES 0
@@ -130,6 +134,8 @@ class TRange;
 	#define DOUBLE_THRESH_QUAT_NORMALIZED			UE_PRIVATE_MATH_DEPRECATION(DOUBLE_THRESH_QUAT_NORMALIZED			, UE_DOUBLE_THRESH_QUAT_NORMALIZED			) UE_DOUBLE_THRESH_QUAT_NORMALIZED
 #endif
 
+#undef UE_INCLUDETOOL_IGNORE_INCONSISTENT_STATE
+
 #define UE_PI 					(3.1415926535897932f)	/* Extra digits if needed: 3.1415926535897932384626433832795f */
 #define UE_SMALL_NUMBER			(1.e-8f)
 #define UE_KINDA_SMALL_NUMBER	(1.e-4f)
