Add GameFeatureTypesFwd.h to forward ignore list

diff --git a/Engine/Source/Programs/IncludeTool/IncludeTool/Rules.cs b/Engine/Source/Programs/IncludeTool/IncludeTool/Rules.cs
--- a/Engine/Source/Programs/IncludeTool/IncludeTool/Rules.cs
+++ b/Engine/Source/Programs/IncludeTool/IncludeTool/Rules.cs
@@ -431,6 +431,7 @@ namespace IncludeTool
 			"/Engine/Source/Runtime/Experimental/Chaos/Public/PhysicsProxy/JointConstraintProxyFwd.h", // invalid forward declaration - 'namespace Chaos'
 			"/Engine/Source/Runtime/Experimental/Chaos/Public/PhysicsProxy/SingleParticlePhysicsProxyFwd.h", // invalid forward declaration - 'namespace Chaos'
 			"/Engine/Source/Runtime/Experimental/Interchange/Engine/Public/InterchangeEngineFwd.h", // invalid forward declaration - 'namespace UE'
+			"/Engine/Plugins/Experimental/GameFeatures/Source/GameFeatures/Public/GameFeatureTypesFwd.h", //  invalid forward declaration - 'namespace GameFeaturePluginStatePrivate'
 		};
 
 		/// <summary>
