IncludeTool: Workaround UE-150382

#rnx

diff --git a/Engine/Source/Programs/IncludeTool/IncludeTool/SymbolTable.cs b/Engine/Source/Programs/IncludeTool/IncludeTool/SymbolTable.cs
--- a/Engine/Source/Programs/IncludeTool/IncludeTool/SymbolTable.cs
+++ b/Engine/Source/Programs/IncludeTool/IncludeTool/SymbolTable.cs
@@ -176,10 +176,17 @@ namespace IncludeTool
 					int NumTypes = Symbols.Select(x => x.Type).Where(x => x != SymbolType.Macro).Distinct().Count();
 					if(NumTypes > 1)
 					{
+						// UE-150382 workaround
+						// It appears that symbols are being added to Lookup whether or not branches have been taken
+						if (Symbols.First().Fragment.File.Location.FullName.Contains("TaskGraphInterfaces.h"))
+						{
+							continue;
+						}
+
 						Log.WriteLine("warning: conflicting declarations of '{0}':", SymbolName);
 						foreach(Symbol Symbol in Symbols)
 						{
-							Log.WriteLine("  {0} in {1}", Symbol.Type, Symbol.Fragment);
+							Log.WriteLine($"  {Symbol.Type} at {Symbol.Fragment.File}:{Symbol.Location.LineIdx+1}");
 						}
 					}
 				}
