IncludeTool: Print an error message when a include cycle is detected.

#rnx

diff --git a/Engine/Source/Programs/IncludeTool/IncludeTool/Program.cs b/Engine/Source/Programs/IncludeTool/IncludeTool/Program.cs
--- a/Engine/Source/Programs/IncludeTool/IncludeTool/Program.cs
+++ b/Engine/Source/Programs/IncludeTool/IncludeTool/Program.cs
@@ -380,7 +380,8 @@ namespace IncludeTool
 					{
 						for(int Idx = 0; Idx < Cycles.Count; Idx++)
 						{
-							Log.WriteLine("warning: include cycle: {0}: {1}", Idx, Cycles[Idx].ToString());
+							Log.WriteLine("Error: include cycle: {0}: {1}", Idx, Cycles[Idx].ToString());
+							Log.WriteLine("Error: Unable to continue until this cycle has been removed.");
 						}
 						return 1;
 					}
