IncludeTool: Launch UBT via dotnet UnrealBuildTool.dll

Adds references to EpicGames.Core and EpicGames.Build; removes duplicate code only when there are explicit name collisions.

#rnx

diff --git a/Engine/Source/Programs/IncludeTool/IncludeTool/BuildTarget.cs b/Engine/Source/Programs/IncludeTool/IncludeTool/BuildTarget.cs
--- a/Engine/Source/Programs/IncludeTool/IncludeTool/BuildTarget.cs
+++ b/Engine/Source/Programs/IncludeTool/IncludeTool/BuildTarget.cs
@@ -1,6 +1,7 @@
 // Copyright Epic Games, Inc. All Rights Reserved.
 
 using IncludeTool.Support;
+using EpicGames.Core;
 using System;
 using System.Collections.Generic;
 using System.Diagnostics;
@@ -92,9 +93,9 @@ namespace IncludeTool
 		public List<BuildModule> Modules = new List<BuildModule>();
 		public Dictionary<string, BuildModule> NameToModule = new Dictionary<string, BuildModule>();
 
-		public static BuildTarget Read(string FileName)
+		public static BuildTarget Read(FileReference File)
 		{
-			JsonObject Object = JsonObject.Read(FileName);
+			JsonObject Object = JsonObject.Read(File);
 
 			BuildTarget Target = new BuildTarget();
 			Target.Name = Object.GetStringField("Name");
diff --git a/Engine/Source/Programs/IncludeTool/IncludeTool/CompileEnvironment.cs b/Engine/Source/Programs/IncludeTool/IncludeTool/CompileEnvironment.cs
--- a/Engine/Source/Programs/IncludeTool/IncludeTool/CompileEnvironment.cs
+++ b/Engine/Source/Programs/IncludeTool/IncludeTool/CompileEnvironment.cs
@@ -1,6 +1,7 @@
 // Copyright Epic Games, Inc. All Rights Reserved.
 
 using IncludeTool.Support;
+using EpicGames.Core;
 using System;
 using System.Collections.Generic;
 using System.IO;
diff --git a/Engine/Source/Programs/IncludeTool/IncludeTool/IncludeTool.csproj b/Engine/Source/Programs/IncludeTool/IncludeTool/IncludeTool.csproj
--- a/Engine/Source/Programs/IncludeTool/IncludeTool/IncludeTool.csproj
+++ b/Engine/Source/Programs/IncludeTool/IncludeTool/IncludeTool.csproj
@@ -42,7 +42,7 @@
   </PropertyGroup>
   <ItemGroup>
     <PackageReference Include="Microsoft.CodeAnalysis.CSharp" Version="3.7.0" />
-    <PackageReference Include="Microsoft.Win32.Registry" Version="4.7.0" />
+    <PackageReference Include="Microsoft.Win32.Registry" Version="5.0.0" />
     <PackageReference Include="System.Management" Version="4.7.0" />
     <PackageReference Include="System.Security.Cryptography.Csp" Version="4.3.0" />
     <PackageReference Include="System.ServiceProcess.ServiceController" Version="4.7.0" />
@@ -54,6 +54,10 @@
     <!-- Required for Hololens -->
     <PackageReference Include="System.Reflection.MetadataLoadContext" Version="4.7.2" /> 
   </ItemGroup>
+  <ItemGroup>
+    <ProjectReference Include="..\..\Shared\EpicGames.Core\EpicGames.Core.csproj" />
+    <ProjectReference Include="..\..\Shared\EpicGames.Build\EpicGames.Build.csproj" />
+  </ItemGroup>
   <ItemGroup>
     <Reference Include="fastJSON">
       <HintPath>..\..\..\..\Binaries\ThirdParty\FastJson\netstandard2.0\FastJson.dll</HintPath>
diff --git a/Engine/Source/Programs/IncludeTool/IncludeTool/Preprocessor.cs b/Engine/Source/Programs/IncludeTool/IncludeTool/Preprocessor.cs
--- a/Engine/Source/Programs/IncludeTool/IncludeTool/Preprocessor.cs
+++ b/Engine/Source/Programs/IncludeTool/IncludeTool/Preprocessor.cs
@@ -7,6 +7,7 @@ using System.Text;
 using System.Threading.Tasks;
 using IncludeTool;
 using IncludeTool.Support;
+using EpicGames.Core;
 using System.Diagnostics;
 using System.IO;
 using System.Collections;
diff --git a/Engine/Source/Programs/IncludeTool/IncludeTool/Program.cs b/Engine/Source/Programs/IncludeTool/IncludeTool/Program.cs
--- a/Engine/Source/Programs/IncludeTool/IncludeTool/Program.cs
+++ b/Engine/Source/Programs/IncludeTool/IncludeTool/Program.cs
@@ -2,6 +2,8 @@
 
 using IncludeTool.Reports;
 using IncludeTool.Support;
+using EpicGames.Core;
+using UnrealBuildBase;
 using System;
 using System.Collections;
 using System.Collections.Concurrent;
@@ -273,7 +275,7 @@ namespace IncludeTool
 
 			// Parse the command line options
 			CommandLineOptions Options = new CommandLineOptions();
-			if(!CommandLine.Parse(Args, Options))
+			if(!Support.CommandLine.Parse(Args, Options))
 			{
 				return 1;
 			}
@@ -307,7 +309,7 @@ namespace IncludeTool
 				}
 
 				// Read the exported target information
-				BuildTarget Target = BuildTarget.Read(TaskListFile.ChangeExtension(".json").FullName);
+				BuildTarget Target = BuildTarget.Read(TaskListFile.ChangeExtension(".json"));
 
 				/// Cache all the files
 				Log.WriteLine("Caching contents of {0}...", InputDir.FullName);
@@ -903,14 +905,13 @@ namespace IncludeTool
 				}
 			}
 
-			DirectoryReference WorkingDir = DirectoryReference.Combine(RootDir, "Engine");
-
-			FileReference UnrealBuildTool = FileReference.Combine(RootDir, "Engine", "Binaries", "DotNET", "UnrealBuildTool", "UnrealBuildTool.exe");
-			if (Utility.Run(UnrealBuildTool, String.Format("-Mode=JsonExport {0} {1} {2}{3} -disableunity -xgeexport -nobuilduht -nopch -nodebuginfo -define:UE_INCLUDE_TOOL=1 -execcodegenactions -outputfile=\"{4}\"", Target, Configuration, Platform, Precompile? " -precompile" : "", TaskListFile.ChangeExtension(".json").FullName), WorkingDir, Log) != 0)
+			if (Utility.Run(Unreal.DotnetPath, String.Format("{0} -Mode=JsonExport {1} {2} {3}{4} -disableunity -xgeexport -nobuilduht -nopch -nodebuginfo -define:UE_INCLUDE_TOOL=1 -execcodegenactions -outputfile=\"{5}\"", 
+				Unreal.UnrealBuildToolDllPath, Target, Configuration, Platform, Precompile? " -precompile" : "", TaskListFile.ChangeExtension(".json").FullName), Unreal.EngineDirectory, Log) != 0)
 			{
 				throw new Exception("UnrealBuildTool failed");
 			}
-			if (Utility.Run(UnrealBuildTool, String.Format("{0} {1} {2}{3} -disableunity -xgeexport -nobuilduht -nopch -nodebuginfo -define:UE_INCLUDE_TOOL=1", Target, Configuration, Platform, Precompile? " -precompile" : ""), WorkingDir, Log) != 0)
+			if (Utility.Run(Unreal.DotnetPath, String.Format("{0} {1} {2} {3}{4} -disableunity -xgeexport -nobuilduht -nopch -nodebuginfo -define:UE_INCLUDE_TOOL=1", 
+				Unreal.UnrealBuildToolDllPath, Target, Configuration, Platform, Precompile? " -precompile" : ""), Unreal.EngineDirectory, Log) != 0)
 			{
 				throw new Exception("UnrealBuildTool failed");
 			}
diff --git a/Engine/Source/Programs/IncludeTool/IncludeTool/Reports/ComplexityReport.cs b/Engine/Source/Programs/IncludeTool/IncludeTool/Reports/ComplexityReport.cs
--- a/Engine/Source/Programs/IncludeTool/IncludeTool/Reports/ComplexityReport.cs
+++ b/Engine/Source/Programs/IncludeTool/IncludeTool/Reports/ComplexityReport.cs
@@ -1,6 +1,7 @@
 // Copyright Epic Games, Inc. All Rights Reserved.
 
 using IncludeTool.Support;
+using EpicGames.Core;
 using System;
 using System.Collections.Generic;
 using System.Diagnostics;
diff --git a/Engine/Source/Programs/IncludeTool/IncludeTool/Reports/FragmentTimingReport.cs b/Engine/Source/Programs/IncludeTool/IncludeTool/Reports/FragmentTimingReport.cs
--- a/Engine/Source/Programs/IncludeTool/IncludeTool/Reports/FragmentTimingReport.cs
+++ b/Engine/Source/Programs/IncludeTool/IncludeTool/Reports/FragmentTimingReport.cs
@@ -1,6 +1,7 @@
 // Copyright Epic Games, Inc. All Rights Reserved.
 
 using IncludeTool.Support;
+using EpicGames.Core;
 using System;
 using System.Collections.Generic;
 using System.Diagnostics;
diff --git a/Engine/Source/Programs/IncludeTool/IncludeTool/Reports/PchReport.cs b/Engine/Source/Programs/IncludeTool/IncludeTool/Reports/PchReport.cs
--- a/Engine/Source/Programs/IncludeTool/IncludeTool/Reports/PchReport.cs
+++ b/Engine/Source/Programs/IncludeTool/IncludeTool/Reports/PchReport.cs
@@ -1,6 +1,7 @@
 // Copyright Epic Games, Inc. All Rights Reserved.
 
 using IncludeTool.Support;
+using EpicGames.Core;
 using System;
 using System.Collections.Generic;
 using System.IO;
diff --git a/Engine/Source/Programs/IncludeTool/IncludeTool/Reports/SymbolReport.cs b/Engine/Source/Programs/IncludeTool/IncludeTool/Reports/SymbolReport.cs
--- a/Engine/Source/Programs/IncludeTool/IncludeTool/Reports/SymbolReport.cs
+++ b/Engine/Source/Programs/IncludeTool/IncludeTool/Reports/SymbolReport.cs
@@ -6,6 +6,7 @@ using System.Linq;
 using System.Text;
 using System.Threading.Tasks;
 using IncludeTool.Support;
+using EpicGames.Core;
 using System.IO;
 
 namespace IncludeTool
diff --git a/Engine/Source/Programs/IncludeTool/IncludeTool/Rules.cs b/Engine/Source/Programs/IncludeTool/IncludeTool/Rules.cs
--- a/Engine/Source/Programs/IncludeTool/IncludeTool/Rules.cs
+++ b/Engine/Source/Programs/IncludeTool/IncludeTool/Rules.cs
@@ -1,6 +1,7 @@
 // Copyright Epic Games, Inc. All Rights Reserved.
 
 using IncludeTool.Support;
+using EpicGames.Core;
 using System;
 using System.Collections.Generic;
 using System.IO;
diff --git a/Engine/Source/Programs/IncludeTool/IncludeTool/SequenceProbe.cs b/Engine/Source/Programs/IncludeTool/IncludeTool/SequenceProbe.cs
--- a/Engine/Source/Programs/IncludeTool/IncludeTool/SequenceProbe.cs
+++ b/Engine/Source/Programs/IncludeTool/IncludeTool/SequenceProbe.cs
@@ -8,6 +8,7 @@ using System.Threading.Tasks;
 using System.IO;
 using System.Diagnostics;
 using IncludeTool.Support;
+using EpicGames.Core;
 
 namespace IncludeTool
 {
diff --git a/Engine/Source/Programs/IncludeTool/IncludeTool/SourceFile.cs b/Engine/Source/Programs/IncludeTool/IncludeTool/SourceFile.cs
--- a/Engine/Source/Programs/IncludeTool/IncludeTool/SourceFile.cs
+++ b/Engine/Source/Programs/IncludeTool/IncludeTool/SourceFile.cs
@@ -1,6 +1,7 @@
 // Copyright Epic Games, Inc. All Rights Reserved.
 
 using IncludeTool.Support;
+using EpicGames.Core;
 using System;
 using System.Collections.Generic;
 using System.Diagnostics;
diff --git a/Engine/Source/Programs/IncludeTool/IncludeTool/SourceFragment.cs b/Engine/Source/Programs/IncludeTool/IncludeTool/SourceFragment.cs
--- a/Engine/Source/Programs/IncludeTool/IncludeTool/SourceFragment.cs
+++ b/Engine/Source/Programs/IncludeTool/IncludeTool/SourceFragment.cs
@@ -1,6 +1,7 @@
 ﻿// Copyright Epic Games, Inc. All Rights Reserved.
 
 using IncludeTool.Support;
+using EpicGames.Core;
 using System;
 using System.Collections.Generic;
 using System.IO;
diff --git a/Engine/Source/Programs/IncludeTool/IncludeTool/Support/FileFilter.cs b/Engine/Source/Programs/IncludeTool/IncludeTool/Support/FileFilter.cs
deleted file mode 100644
diff --git a/Engine/Source/Programs/IncludeTool/IncludeTool/Support/FileSystemReference.cs b/Engine/Source/Programs/IncludeTool/IncludeTool/Support/FileSystemReference.cs
deleted file mode 100644
diff --git a/Engine/Source/Programs/IncludeTool/IncludeTool/Support/FileSystemReferenceExtensions.cs b/Engine/Source/Programs/IncludeTool/IncludeTool/Support/FileSystemReferenceExtensions.cs
new file mode 100644
--- /dev/null
+++ b/Engine/Source/Programs/IncludeTool/IncludeTool/Support/FileSystemReferenceExtensions.cs
@@ -0,0 +1,25 @@
+using System;
+using System.Collections.Generic;
+using System.Linq;
+using System.Text;
+using System.Threading.Tasks;
+using EpicGames.Core;
+
+namespace IncludeTool.Support
+{
+	/// <summary>
+	/// Extension methods to smooth the transition from IncludeTool.Support.*Reference to EpicGames.Core.*Reference
+	/// Don't add to this; rather, remove it and fix up the call sites for consistency across projects.
+	/// </summary>
+	internal static class FileSystemReferenceExtensions
+	{
+		public static bool Exists(this FileReference File) => FileReference.Exists(File);
+		public static void Delete(this FileReference File) => FileReference.Delete(File);
+		public static void CreateDirectory(this DirectoryReference Directory) => DirectoryReference.CreateDirectory(Directory);
+
+		public static bool Exists(this DirectoryReference Directory) => DirectoryReference.Exists(Directory); 
+		public static IEnumerable<FileReference> EnumerateFileReferences(this DirectoryReference Directory) => DirectoryReference.EnumerateFiles(Directory);
+		public static IEnumerable<FileReference> EnumerateFileReferences(this DirectoryReference Directory, string Pattern) => DirectoryReference.EnumerateFiles(Directory, Pattern);
+		public static IEnumerable<DirectoryReference> EnumerateDirectoryReferences(this DirectoryReference Directory, string Pattern, System.IO.SearchOption Option) => DirectoryReference.EnumerateDirectories(Directory, Pattern, Option);
+	}
+}
diff --git a/Engine/Source/Programs/IncludeTool/IncludeTool/Support/JsonObject.cs b/Engine/Source/Programs/IncludeTool/IncludeTool/Support/JsonObject.cs
deleted file mode 100644
diff --git a/Engine/Source/Programs/IncludeTool/IncludeTool/Support/LineBasedTextWriter.cs b/Engine/Source/Programs/IncludeTool/IncludeTool/Support/LineBasedTextWriter.cs
--- a/Engine/Source/Programs/IncludeTool/IncludeTool/Support/LineBasedTextWriter.cs
+++ b/Engine/Source/Programs/IncludeTool/IncludeTool/Support/LineBasedTextWriter.cs
@@ -1,5 +1,6 @@
 // Copyright Epic Games, Inc. All Rights Reserved.
 
+using EpicGames.Core;
 using System;
 using System.Collections.Generic;
 using System.Linq;
diff --git a/Engine/Source/Programs/IncludeTool/IncludeTool/Support/ThreadPoolWorkQueue.cs b/Engine/Source/Programs/IncludeTool/IncludeTool/Support/ThreadPoolWorkQueue.cs
deleted file mode 100644
diff --git a/Engine/Source/Programs/IncludeTool/IncludeTool/Support/Utility.cs b/Engine/Source/Programs/IncludeTool/IncludeTool/Support/Utility.cs
--- a/Engine/Source/Programs/IncludeTool/IncludeTool/Support/Utility.cs
+++ b/Engine/Source/Programs/IncludeTool/IncludeTool/Support/Utility.cs
@@ -1,5 +1,6 @@
 // Copyright Epic Games, Inc. All Rights Reserved.
 
+using EpicGames.Core;
 using System;
 using System.Collections.Generic;
 using System.Diagnostics;
diff --git a/Engine/Source/Programs/IncludeTool/IncludeTool/Workspace.cs b/Engine/Source/Programs/IncludeTool/IncludeTool/Workspace.cs
--- a/Engine/Source/Programs/IncludeTool/IncludeTool/Workspace.cs
+++ b/Engine/Source/Programs/IncludeTool/IncludeTool/Workspace.cs
@@ -1,6 +1,7 @@
 // Copyright Epic Games, Inc. All Rights Reserved.
 
 using IncludeTool.Support;
+using EpicGames.Core;
 using System;
 using System.Collections.Concurrent;
 using System.Collections.Generic;
@@ -321,7 +322,7 @@ namespace IncludeTool
 	/// <summary>
 	/// Worker class to find files in the workspace
 	/// </summary>
-	class WorkspaceScanner : IThreadPoolWorker
+	class WorkspaceScanner
 	{
 		/// <summary>
 		/// The current directory being searched
@@ -363,7 +364,7 @@ namespace IncludeTool
 					}
 					else
 					{
-						Queue.Enqueue(new WorkspaceScanner(ChildDirectory));
+						Queue.Enqueue(() => new WorkspaceScanner(ChildDirectory).Run(Queue));
 					}
 				}
 				CurrentDirectory = NextDirectory;
@@ -374,7 +375,7 @@ namespace IncludeTool
 	/// <summary>
 	/// Worker class to find files in the workspace
 	/// </summary>
-	class WorkspaceReferenceScanner : IThreadPoolWorker
+	class WorkspaceReferenceScanner
 	{
 		/// <summary>
 		/// The current directory being searched
@@ -424,7 +425,7 @@ namespace IncludeTool
 			{
 				if (Rules.SearchDirectoryForSource(Directory.NormalizedPathFromBranchRoot))
 				{
-					Queue.Enqueue(new WorkspaceReferenceScanner(Directory, FoundFiles));
+					Queue.Enqueue(() => new WorkspaceReferenceScanner(Directory, FoundFiles).Run(Queue));
 				}
 			}
 		}
@@ -479,7 +480,7 @@ namespace IncludeTool
 		{
 			using (ThreadPoolWorkQueue WorkQueue = new ThreadPoolWorkQueue())
 			{
-				WorkQueue.Enqueue(new WorkspaceScanner(Directory));
+				WorkQueue.Enqueue(() => new WorkspaceScanner(Directory).Run(WorkQueue));
 			}
 		}
 
