IncludeTool: Update use of FileReference and DirectoryReference to EpicGames.Core APIs

#rnx

diff --git a/Engine/Source/Programs/IncludeTool/IncludeTool/Preprocessor.cs b/Engine/Source/Programs/IncludeTool/IncludeTool/Preprocessor.cs
--- a/Engine/Source/Programs/IncludeTool/IncludeTool/Preprocessor.cs
+++ b/Engine/Source/Programs/IncludeTool/IncludeTool/Preprocessor.cs
@@ -2018,15 +2018,15 @@ namespace IncludeTool
 		/// <param name="Log">Writer for log output</param>
 		public static void WritePreprocessedFiles(DirectoryReference WorkingDir, DirectoryReference PreprocessedDir, IEnumerable<KeyValuePair<FileReference, CompileEnvironment>> FilePairs, LineBasedTextWriter Log)
 		{
-			PreprocessedDir.CreateDirectory();
+			DirectoryReference.CreateDirectory(PreprocessedDir);
 
 			// Create separate directories for the compiler and include tool
 			DirectoryReference CompilerDir = DirectoryReference.Combine(PreprocessedDir, "Compiler");
-			CompilerDir.CreateDirectory();
+			DirectoryReference.CreateDirectory(CompilerDir);
 			DirectoryReference CompilerNormalizedDir = DirectoryReference.Combine(PreprocessedDir, "CompilerNormalized");
-			CompilerNormalizedDir.CreateDirectory();
+			DirectoryReference.CreateDirectory(CompilerNormalizedDir);
 			DirectoryReference ToolDir = DirectoryReference.Combine(PreprocessedDir, "Tool");
-			ToolDir.CreateDirectory();
+			DirectoryReference.CreateDirectory(ToolDir);
 
 			// Write out the prelude file, which we can use to scrape the compiler settings
 			FileReference PreludeFile = FileReference.Combine(PreprocessedDir, "Prelude.cpp");
diff --git a/Engine/Source/Programs/IncludeTool/IncludeTool/Program.cs b/Engine/Source/Programs/IncludeTool/IncludeTool/Program.cs
--- a/Engine/Source/Programs/IncludeTool/IncludeTool/Program.cs
+++ b/Engine/Source/Programs/IncludeTool/IncludeTool/Program.cs
@@ -282,7 +282,7 @@ namespace IncludeTool
 
 			// Make sure we've got an output directory
 			DirectoryReference WorkingDir = new DirectoryReference(Options.WorkingDir);
-			WorkingDir.CreateDirectory();
+			DirectoryReference.CreateDirectory(WorkingDir);
 
 			// Create an output log
 			FileReference LogFile = FileReference.Combine(WorkingDir, "Log.txt");
@@ -301,10 +301,10 @@ namespace IncludeTool
 
 				// Generate the exported task list 
 				FileReference TaskListFile = FileReference.Combine(WorkingDir, "Targets", String.Format("{0} {1} {2}.xml", Options.Target, Options.Platform, Options.Configuration));
-				if (Options.CleanTaskList || !TaskListFile.Exists())
+				if (Options.CleanTaskList || !FileReference.Exists(TaskListFile))
 				{
 					Log.WriteLine("Generating task list for {0} {1} {2}...", Options.Target, Options.Platform, Options.Configuration);
-					TaskListFile.Directory.CreateDirectory();
+					DirectoryReference.CreateDirectory(TaskListFile.Directory);
 					GenerateTaskList(InputDir, Options.Target, Options.Platform, Options.Configuration, Options.Precompile, TaskListFile, Log);
 				}
 
@@ -469,9 +469,9 @@ namespace IncludeTool
 					foreach(string Name in StandardModuleDirectoryNames)
 					{
 						DirectoryReference BaseDir = DirectoryReference.Combine(Module.Directory, Name);
-						if(BaseDir.Exists())
+						if(DirectoryReference.Exists(BaseDir))
 						{
-							foreach(DirectoryReference ChildDir in BaseDir.EnumerateDirectoryReferences("*", SearchOption.AllDirectories))
+							foreach(DirectoryReference ChildDir in DirectoryReference.EnumerateDirectories(BaseDir, "*", SearchOption.AllDirectories))
 							{
 								RebaseIncludePaths[ChildDir] = BaseDir;
 							}
@@ -503,7 +503,7 @@ namespace IncludeTool
 						foreach(string Name in StandardModuleDirectoryNames)
 						{
 							DirectoryReference IncludePath = DirectoryReference.Combine(Pair.Key.Module.Directory, Name);
-							if(IncludePath.Exists())
+							if(DirectoryReference.Exists(IncludePath))
 							{
 								IncludePaths.Add(IncludePath);
 							}
@@ -742,7 +742,7 @@ namespace IncludeTool
 
 						// Create the intermediate directory
 						DirectoryReference IntermediateDir = DirectoryReference.Combine(WorkingDir, "Intermediate");
-						IntermediateDir.CreateDirectory();
+						DirectoryReference.CreateDirectory(IntermediateDir);
 
 						// If we're just checking fragments, do so now
 						if(Options.Mode == ToolMode.Verify)
@@ -760,7 +760,7 @@ namespace IncludeTool
 							Dictionary<FileReference, DateTime> FileToTime = new Dictionary<FileReference, DateTime>();
 							if(Options.NumShards > 1)
 							{
-								foreach(FileReference IntermediateFile in IntermediateDir.EnumerateFileReferences())
+								foreach(FileReference IntermediateFile in DirectoryReference.EnumerateFiles(IntermediateDir))
 								{
 									DateTime LastWriteTime = new FileInfo(IntermediateFile.FullName).LastWriteTimeUtc;
 									FileToTime[IntermediateFile] = LastWriteTime;
@@ -774,10 +774,10 @@ namespace IncludeTool
 							if(Options.NumShards > 1)
 							{
 								DirectoryReference ShardDir = DirectoryReference.Combine(WorkingDir, "Shard");
-								ShardDir.CreateDirectory();
+								DirectoryReference.CreateDirectory(ShardDir);
 								Log.WriteLine("Copying shard outputs to {0}...", ShardDir);
 
-								foreach(FileReference IntermediateFile in IntermediateDir.EnumerateFileReferences())
+								foreach(FileReference IntermediateFile in DirectoryReference.EnumerateFiles(IntermediateDir))
 								{
 									if(IntermediateFile.HasExtension(".txt") || IntermediateFile.HasExtension(".response") || IntermediateFile.HasExtension(".state") || IntermediateFile.HasExtension(".permutation"))
 									{
@@ -897,11 +897,11 @@ namespace IncludeTool
 		static void GenerateTaskList(DirectoryReference RootDir, string Target, TargetPlatform Platform, TargetConfiguration Configuration, bool Precompile, FileReference TaskListFile, LineBasedTextWriter Log)
 		{
 			DirectoryReference IntermediateDir = DirectoryReference.Combine(RootDir, "Engine", "Intermediate", "Build");
-			if(IntermediateDir.Exists())
+			if(DirectoryReference.Exists(IntermediateDir))
 			{
-				foreach(FileReference IntermediateFile in IntermediateDir.EnumerateFileReferences("UBTExport*.xml"))
+				foreach(FileReference IntermediateFile in DirectoryReference.EnumerateFiles(IntermediateDir, "UBTExport*.xml"))
 				{
-					IntermediateFile.Delete();
+					FileReference.Delete(IntermediateFile);
 				}
 			}
 
@@ -915,7 +915,7 @@ namespace IncludeTool
 			{
 				throw new Exception("UnrealBuildTool failed");
 			}
-			TaskListFile.Directory.CreateDirectory();
+			DirectoryReference.CreateDirectory(TaskListFile.Directory);
 			File.Copy(FileReference.Combine(RootDir, "Engine", "Intermediate", "Build", "UBTExport.000.xge.xml").FullName, TaskListFile.FullName, true);
 		}
 
@@ -1266,7 +1266,7 @@ namespace IncludeTool
 		static void WriteFragments(IEnumerable<SourceFile> Files, DirectoryReference OutputDir)
 		{
 			// Create the output directory
-			OutputDir.CreateDirectory();
+			DirectoryReference.CreateDirectory(OutputDir);
 
 			// Assign names to all the fragments
 			HashSet<SourceFile> VisitedFiles = new HashSet<SourceFile>();
@@ -1585,7 +1585,7 @@ namespace IncludeTool
 				string Contents = Writer.ToString();
 
 				FileReference NewLocation = FileReference.Combine(OutputDir, OutputFile.Location.MakeRelativeTo(InputDir));
-				if(!NewLocation.Exists() || File.ReadAllText(NewLocation.FullName).TrimEnd() != Contents.TrimEnd())
+				if(!FileReference.Exists(NewLocation) || File.ReadAllText(NewLocation.FullName).TrimEnd() != Contents.TrimEnd())
 				{
 					OutputFileContents.Add(NewLocation, Contents);
 				}
@@ -1608,7 +1608,7 @@ namespace IncludeTool
 		static bool WriteOutputFiles(DirectoryReference OutputDir, Dictionary<FileReference, string> OutputFileContents, bool bWithP4, LineBasedTextWriter Log)
 		{
 			// Create the output directory
-			OutputDir.CreateDirectory();
+			DirectoryReference.CreateDirectory(OutputDir);
 
 			// Check out any that are read-only from perforce
 			if(bWithP4)
@@ -1643,7 +1643,7 @@ namespace IncludeTool
 			foreach(KeyValuePair<FileReference, string> OutputFile in OutputFileContents)
 			{
 				Log.WriteLine("{0}", OutputFile.Key);
-				OutputFile.Key.Directory.CreateDirectory();
+				DirectoryReference.CreateDirectory(OutputFile.Key.Directory);
 				File.WriteAllText(OutputFile.Key.FullName, OutputFile.Value);
 			}
 			return true;
diff --git a/Engine/Source/Programs/IncludeTool/IncludeTool/Reports/FragmentTimingReport.cs b/Engine/Source/Programs/IncludeTool/IncludeTool/Reports/FragmentTimingReport.cs
--- a/Engine/Source/Programs/IncludeTool/IncludeTool/Reports/FragmentTimingReport.cs
+++ b/Engine/Source/Programs/IncludeTool/IncludeTool/Reports/FragmentTimingReport.cs
@@ -261,7 +261,7 @@ namespace IncludeTool.Reports
 
 			// Create an intermediate directory
 			DirectoryReference IntermediateDir = DirectoryReference.Combine(WorkingDir, "Timing");
-			IntermediateDir.CreateDirectory();
+			DirectoryReference.CreateDirectory(IntermediateDir);
 
 			// Map of unique fragment to timing data
 			Dictionary<SourceFragment, FragmentTimingData> FragmentToTimingData = new Dictionary<SourceFragment, FragmentTimingData>();
@@ -318,7 +318,7 @@ namespace IncludeTool.Reports
 			}
 
 			// Read any existing shard timing data in the output folder
-			foreach(FileReference IntermediateFile in IntermediateDir.EnumerateFileReferences("*.csv"))
+			foreach(FileReference IntermediateFile in DirectoryReference.EnumerateFiles(IntermediateDir, "*.csv"))
 			{
 				string[] Lines = File.ReadAllLines(IntermediateFile.FullName);
 				foreach(string Line in Lines.Skip(1))
diff --git a/Engine/Source/Programs/IncludeTool/IncludeTool/SequenceProbe.cs b/Engine/Source/Programs/IncludeTool/IncludeTool/SequenceProbe.cs
--- a/Engine/Source/Programs/IncludeTool/IncludeTool/SequenceProbe.cs
+++ b/Engine/Source/Programs/IncludeTool/IncludeTool/SequenceProbe.cs
@@ -236,7 +236,7 @@ namespace IncludeTool
 			}
 
 			// Check to see if an existing version of the task exists which we can continue
-			if(Type == SequenceProbeType.Optimize && TaskStateFile.Exists())
+			if(Type == SequenceProbeType.Optimize && FileReference.Exists(TaskStateFile))
 			{
 				// Try to read the old task
 				SequenceWorker OldWorker;
@@ -315,9 +315,9 @@ namespace IncludeTool
 			{
 				Worker.Serialize(TaskStateFile.FullName);
 
-				PermutationFile.Delete();
-				SummaryLogFile.Delete();
-				CompileLogFile.Delete();
+				FileReference.Delete(PermutationFile);
+				FileReference.Delete(SummaryLogFile);
+				FileReference.Delete(CompileLogFile);
 			}
 		}
 
@@ -573,7 +573,7 @@ namespace IncludeTool
 			// Read the new state
 			FileReference OutputFile = new FileReference(TaskStateFile.FullName + ".out");
 			SequenceWorker NewWorker = SequenceWorker.Deserialize(OutputFile.FullName);
-			OutputFile.Delete();
+			FileReference.Delete(OutputFile);
 
 			// Make sure the exit code reflects the failure state. XGE can sometimes fail transferring back.
 			if(ExitCode == 0 && !NewWorker.bResult)
diff --git a/Engine/Source/Programs/IncludeTool/IncludeTool/Support/FileSystemReferenceExtensions.cs b/Engine/Source/Programs/IncludeTool/IncludeTool/Support/FileSystemReferenceExtensions.cs
deleted file mode 100644
diff --git a/Engine/Source/Programs/IncludeTool/IncludeTool/Support/Utility.cs b/Engine/Source/Programs/IncludeTool/IncludeTool/Support/Utility.cs
--- a/Engine/Source/Programs/IncludeTool/IncludeTool/Support/Utility.cs
+++ b/Engine/Source/Programs/IncludeTool/IncludeTool/Support/Utility.cs
@@ -28,7 +28,7 @@ namespace IncludeTool.Support
 			foreach(string SearchPath in Environment.GetEnvironmentVariable("PATH").Split(Path.PathSeparator))
 			{
 				FileReference CandidateLocation = FileReference.Combine(new DirectoryReference(SearchPath), "p4.exe");
-				if(CandidateLocation.Exists())
+				if(FileReference.Exists(CandidateLocation))
 				{
 					Location = CandidateLocation;
 					return true;
