IncludeTool fixes for new thirdparty headers.
1) Skip over "__has_include_next" preprocessor statements. They have the same format as __has_include statements: #if __has_include_next(<filename.h>)
2) Skip over __has_extension statements: #if __has_extension(featurename)

diff --git a/Engine/Source/Programs/IncludeTool/IncludeTool/PreprocessorExpression.cs b/Engine/Source/Programs/IncludeTool/IncludeTool/PreprocessorExpression.cs
--- a/Engine/Source/Programs/IncludeTool/IncludeTool/PreprocessorExpression.cs
+++ b/Engine/Source/Programs/IncludeTool/IncludeTool/PreprocessorExpression.cs
@@ -340,6 +340,7 @@ namespace IncludeTool
 				case "alignof":
 					throw new NotImplementedException();
 				case "__has_builtin":
+				case "__has_extension":
 				case "__has_feature":
 				case "__has_warning":
 					if(Tokens[Idx + 1].Text != "(" || Tokens[Idx + 3].Text != ")")
@@ -349,6 +350,7 @@ namespace IncludeTool
 					Idx += 4;
 					return 0;
 				case "__has_include":
+				case "__has_include_next":
 					if(Tokens[Idx + 1].Text == "(")
 					{
 						Idx += 4; // 4 = __has_inc, (, "" or <>
