IncludeTool: Handle /imsvc & -isystem include paths in clang command lines

#rnx

diff --git a/Engine/Source/Programs/IncludeTool/IncludeTool/CompileEnvironment.cs b/Engine/Source/Programs/IncludeTool/IncludeTool/CompileEnvironment.cs
--- a/Engine/Source/Programs/IncludeTool/IncludeTool/CompileEnvironment.cs
+++ b/Engine/Source/Programs/IncludeTool/IncludeTool/CompileEnvironment.cs
@@ -315,7 +315,7 @@ namespace IncludeTool
 							{
 								Environment.Definitions.Add(Tokens[Idx].Substring(2));
 							}
-							else if (Tokens[Idx] == "/I" || Tokens[Idx] == "-I")
+							else if (Tokens[Idx] == "/I" || Tokens[Idx] == "-I" || Tokens[Idx] == "/imsvc" || Tokens[Idx] == "-isystem")
 							{
 								string IncludePath = Tokens[++Idx].Replace("//", "/");
 								if (!Path.IsPathRooted(IncludePath))
@@ -333,6 +333,15 @@ namespace IncludeTool
 								}
 								Environment.IncludePaths.Add(new DirectoryReference(IncludePath.ToLowerInvariant()));
 							}
+							else if (Tokens[Idx].StartsWith("-isystem"))
+							{
+								string IncludePath = Tokens[Idx].Substring(8).Replace("//", "/");
+								if (!Path.IsPathRooted(IncludePath))
+								{
+									IncludePath = Path.Combine(BaseDir.FullName, IncludePath);
+								}
+								Environment.IncludePaths.Add(new DirectoryReference(IncludePath.ToLowerInvariant()));
+							}
 							else if(Tokens[Idx].StartsWith("/FI"))
 							{
 								string ForceIncludePath = Tokens[Idx].Substring(3);
