IncludeTool: Handle /external:I as include path arg

#rnx

diff --git a/Engine/Source/Programs/IncludeTool/IncludeTool/CompileEnvironment.cs b/Engine/Source/Programs/IncludeTool/IncludeTool/CompileEnvironment.cs
--- a/Engine/Source/Programs/IncludeTool/IncludeTool/CompileEnvironment.cs
+++ b/Engine/Source/Programs/IncludeTool/IncludeTool/CompileEnvironment.cs
@@ -315,7 +315,7 @@ namespace IncludeTool
 							{
 								Environment.Definitions.Add(Tokens[Idx].Substring(2));
 							}
-							else if (Tokens[Idx] == "/I" || Tokens[Idx] == "-I" || Tokens[Idx] == "/imsvc" || Tokens[Idx] == "-isystem")
+							else if (Tokens[Idx] == "/I" || Tokens[Idx] == "/external:I" || Tokens[Idx] == "-I" || Tokens[Idx] == "/imsvc" || Tokens[Idx] == "-isystem")
 							{
 								string IncludePath = Tokens[++Idx].Replace("//", "/");
 								if (!Path.IsPathRooted(IncludePath))
