IncludeTool: Fix errors building from a directory containing a space.

diff --git a/Engine/Source/Programs/IncludeTool/IncludeTool/Program.cs b/Engine/Source/Programs/IncludeTool/IncludeTool/Program.cs
--- a/Engine/Source/Programs/IncludeTool/IncludeTool/Program.cs
+++ b/Engine/Source/Programs/IncludeTool/IncludeTool/Program.cs
@@ -905,12 +905,12 @@ namespace IncludeTool
 				}
 			}
 
-			if (Utility.Run(Unreal.DotnetPath, String.Format("{0} -Mode=JsonExport {1} {2} {3}{4} -disableunity -xgeexport -nobuilduht -nopch -nodebuginfo -define:UE_INCLUDE_TOOL=1 -execcodegenactions -outputfile=\"{5}\"",
+			if (Utility.Run(Unreal.DotnetPath, String.Format("\"{0}\" -Mode=JsonExport {1} {2} {3}{4} -disableunity -xgeexport -nobuilduht -nopch -nodebuginfo -define:UE_INCLUDE_TOOL=1 -execcodegenactions -outputfile=\"{5}\"",
 				Unreal.UnrealBuildToolDllPath, Target, Configuration, Platform, Precompile? " -precompile" : "", TaskListFile.ChangeExtension(".json").FullName), Unreal.EngineDirectory, Log) != 0)
 			{
 				throw new Exception("UnrealBuildTool failed");
 			}
-			if (Utility.Run(Unreal.DotnetPath, String.Format("{0} {1} {2} {3}{4} -disableunity -xgeexport -nobuilduht -nopch -nodebuginfo -define:UE_INCLUDE_TOOL=1",
+			if (Utility.Run(Unreal.DotnetPath, String.Format("\"{0}\" {1} {2} {3}{4} -disableunity -xgeexport -nobuilduht -nopch -nodebuginfo -define:UE_INCLUDE_TOOL=1",
 				Unreal.UnrealBuildToolDllPath, Target, Configuration, Platform, Precompile? " -precompile" : ""), Unreal.EngineDirectory, Log) != 0)
 			{
 				throw new Exception("UnrealBuildTool failed");
