IncludeTool: Suppress warnings for Rev.Runtime folder.

diff --git a/Engine/Source/Programs/IncludeTool/IncludeTool/Rules.cs b/Engine/Source/Programs/IncludeTool/IncludeTool/Rules.cs
--- a/Engine/Source/Programs/IncludeTool/IncludeTool/Rules.cs
+++ b/Engine/Source/Programs/IncludeTool/IncludeTool/Rules.cs
@@ -197,22 +197,26 @@
 			{
 				return true;
 			}
-			if(NormalizedPath.StartsWith("/engine/plugins/experimental/phya/source/phya/private/phyalib/"))
+			if (NormalizedPath.StartsWith("/engine/plugins/experimental/phya/source/phya/private/phyalib/"))
 			{
 				return true;
 			}
-			if(NormalizedPath.StartsWith("/engine/plugins/runtime/leapmotion/thirdparty/"))
+			if (NormalizedPath.StartsWith("/engine/plugins/runtime/leapmotion/thirdparty/"))
 			{
 				return true;
 			}
-			if(NormalizedPath.EndsWith("/recastmesh.cpp") || NormalizedPath.EndsWith("/recastfilter.cpp") || NormalizedPath.EndsWith("/recastcontour.cpp") || NormalizedPath.EndsWith("/framepro.h") || NormalizedPath.EndsWith("/framepro.cpp") || NormalizedPath.EndsWith("/frameproue4.h") || NormalizedPath.EndsWith("/frameproue4.cpp") || NormalizedPath.EndsWith("/sqlite3.h") || NormalizedPath.EndsWith("/sqlite3.inl") || NormalizedPath.EndsWith("/vorbis_stream_encoder.h") || NormalizedPath.EndsWith("/integral_types.h"))
+			if (NormalizedPath.Contains("/rev.runtime/"))
 			{
 				return true;
 			}
-			if(NormalizedPath.Contains("/thirdparty/rapidjson/"))
+			if (NormalizedPath.EndsWith("/recastmesh.cpp") || NormalizedPath.EndsWith("/recastfilter.cpp") || NormalizedPath.EndsWith("/recastcontour.cpp") || NormalizedPath.EndsWith("/framepro.h") || NormalizedPath.EndsWith("/framepro.cpp") || NormalizedPath.EndsWith("/frameproue4.h") || NormalizedPath.EndsWith("/frameproue4.cpp") || NormalizedPath.EndsWith("/sqlite3.h") || NormalizedPath.EndsWith("/sqlite3.inl") || NormalizedPath.EndsWith("/vorbis_stream_encoder.h") || NormalizedPath.EndsWith("/integral_types.h"))
 			{
 				return true;
 			}
+			if (NormalizedPath.Contains("/thirdparty/rapidjson/"))
+			{
+				return true;
+			}
 			if (NormalizedPath.EndsWith("/lz4.h") || NormalizedPath.EndsWith("/lz4hc.h") || NormalizedPath.EndsWith("/lz4.cpp") || NormalizedPath.EndsWith("/lz4hc.cpp"))
 			{
 				return true;
@@ -578,6 +582,10 @@
 			{
 				return true;
 			}
+			if((File.Flags & SourceFileFlags.External) != 0)
+			{
+				return true;
+			}
 			return false;
 		}
 
