Disable warning for inconsistent type of RigVMTypeIndex.

diff --git a/Engine/Source/Programs/IncludeTool/IncludeTool/SymbolTable.cs b/Engine/Source/Programs/IncludeTool/IncludeTool/SymbolTable.cs
--- a/Engine/Source/Programs/IncludeTool/IncludeTool/SymbolTable.cs
+++ b/Engine/Source/Programs/IncludeTool/IncludeTool/SymbolTable.cs
@@ -178,7 +178,12 @@ namespace IncludeTool
 					{
 						// UE-150382 workaround
 						// It appears that symbols are being added to Lookup whether or not branches have been taken
-						if (Symbols.First().Fragment.File.Location.FullName.Contains("TaskGraphInterfaces.h"))
+						string FileName = Symbols.First().Fragment.File.Location.GetFileName();
+						if (FileName.Equals("TaskGraphInterfaces.h", StringComparison.OrdinalIgnoreCase))
+						{
+							continue;
+						}
+						if (FileName.Equals("RigVMTypeIndex.h", StringComparison.OrdinalIgnoreCase))
 						{
 							continue;
 						}
