Treat new generated headers consistent with ispc generated headers for IncludeTool purposes

diff --git a/Engine/Source/Programs/IncludeTool/IncludeTool/Rules.cs b/Engine/Source/Programs/IncludeTool/IncludeTool/Rules.cs
--- a/Engine/Source/Programs/IncludeTool/IncludeTool/Rules.cs
+++ b/Engine/Source/Programs/IncludeTool/IncludeTool/Rules.cs
@@ -463,7 +463,7 @@ namespace IncludeTool
 			}
 			if(NormalizedPath.IndexOf("/intermediate/") != -1)
 			{
-				if (NormalizedPath.EndsWith("ispc.generated.h"))
+				if (NormalizedPath.EndsWith("ispc.generated.h") || NormalizedPath.EndsWith(".gen.h"))
 				{
 					Flags |= SourceFileFlags.GeneratedHeader | SourceFileFlags.Public;
 				}
