Additional case to make new generated headers treated consistent with ispc

diff --git a/Engine/Source/Programs/IncludeTool/IncludeTool/Rules.cs b/Engine/Source/Programs/IncludeTool/IncludeTool/Rules.cs
--- a/Engine/Source/Programs/IncludeTool/IncludeTool/Rules.cs
+++ b/Engine/Source/Programs/IncludeTool/IncludeTool/Rules.cs
@@ -453,7 +453,7 @@ namespace IncludeTool
 			}
 
 			SourceFileFlags Flags = SourceFileFlags.Standalone;
-			if(NormalizedPath.EndsWith(".inl") || NormalizedPath.EndsWith(".inc") || NormalizedPath.EndsWith(".generated.h"))
+			if(NormalizedPath.EndsWith(".inl") || NormalizedPath.EndsWith(".inc") || NormalizedPath.EndsWith(".generated.h") || NormalizedPath.EndsWith(".gen.h"))
 			{
 				Flags = (Flags | SourceFileFlags.Pinned) & ~SourceFileFlags.Standalone;
 			}
