Permit VNI generated files to have multiple fragments

diff --git a/Engine/Source/Programs/IncludeTool/IncludeTool/Rules.cs b/Engine/Source/Programs/IncludeTool/IncludeTool/Rules.cs
--- a/Engine/Source/Programs/IncludeTool/IncludeTool/Rules.cs
+++ b/Engine/Source/Programs/IncludeTool/IncludeTool/Rules.cs
@@ -453,9 +453,10 @@ namespace IncludeTool
 			}
 
 			SourceFileFlags Flags = SourceFileFlags.Standalone;
-			if(NormalizedPath.EndsWith(".inl") || NormalizedPath.EndsWith(".inc") || NormalizedPath.EndsWith(".generated.h") || NormalizedPath.EndsWith(".gen.h"))
+			if(NormalizedPath.EndsWith(".inl") || NormalizedPath.EndsWith(".inc") || NormalizedPath.EndsWith(".generated.h"))
 			{
-				Flags = (Flags | SourceFileFlags.Pinned) & ~SourceFileFlags.Standalone;
+				Flags &= ~SourceFileFlags.Standalone;
+				Flags |= SourceFileFlags.Pinned;
 			}
 			if(NormalizedPath.IndexOf("/public/") != -1 || NormalizedPath.IndexOf("/classes/") != -1)
 			{
@@ -463,7 +464,7 @@ namespace IncludeTool
 			}
 			if(NormalizedPath.IndexOf("/intermediate/") != -1)
 			{
-				if (NormalizedPath.EndsWith("ispc.generated.h") || NormalizedPath.EndsWith(".gen.h"))
+				if (NormalizedPath.EndsWith("ispc.generated.h"))
 				{
 					Flags |= SourceFileFlags.GeneratedHeader | SourceFileFlags.Public;
 				}
@@ -475,6 +476,11 @@ namespace IncludeTool
 				{
 					Flags |= SourceFileFlags.GeneratedClassesHeader | SourceFileFlags.Public;
 				}
+				else if(NormalizedPath.EndsWith(".gen.h") && NormalizedPath.IndexOf("/vni/") != -1)
+				{
+					Flags &= ~SourceFileFlags.Standalone;
+					Flags |= SourceFileFlags.Pinned | SourceFileFlags.GeneratedHeader | SourceFileFlags.Public | SourceFileFlags.AllowMultipleFragments;
+				}
 			}
 			if(NormalizedPath.EndsWith(".cpp") || NormalizedPath.IndexOf("/windows/") != -1 || NormalizedPath.IndexOf("/linux/") != -1)
 			{
