Modify IncludeTool external file test for NNX plugin

#rnx

diff --git a/Engine/Source/Programs/IncludeTool/IncludeTool/Program.cs b/Engine/Source/Programs/IncludeTool/IncludeTool/Program.cs
--- a/Engine/Source/Programs/IncludeTool/IncludeTool/Program.cs
+++ b/Engine/Source/Programs/IncludeTool/IncludeTool/Program.cs
@@ -1067,6 +1067,7 @@ namespace IncludeTool
 							&& !File.Location.FullName.Contains("lz4")
 							&& !File.Location.FullName.Contains("NeuralNetworkInference")
 							&& !File.Location.FullName.Contains("NNI")
+							&& !File.Location.FullName.Contains("NNX")
 							&& !File.Location.FullName.Contains("PhyaLib")
 							)
 						{
