IncludeTool: Add Engine\Source\Runtime\Core\Public\Containers\VersePathFwd.h to fwd exclude list

#rnx

diff --git a/Engine/Source/Programs/IncludeTool/IncludeTool/Rules.cs b/Engine/Source/Programs/IncludeTool/IncludeTool/Rules.cs
--- a/Engine/Source/Programs/IncludeTool/IncludeTool/Rules.cs
+++ b/Engine/Source/Programs/IncludeTool/IncludeTool/Rules.cs
@@ -437,6 +437,7 @@ namespace IncludeTool
 			"/Engine/Source/Runtime/Experimental/Chaos/Public/PhysicsProxy/SingleParticlePhysicsProxyFwd.h", // invalid forward declaration - 'namespace Chaos'
 			"/Engine/Source/Runtime/Interchange/Engine/Public/InterchangeEngineFwd.h", // invalid forward declaration - 'namespace UE'
 			"/Engine/Plugins/Experimental/GameFeatures/Source/GameFeatures/Public/GameFeatureTypesFwd.h", //  invalid forward declaration - 'namespace GameFeaturePluginStatePrivate'
+			"/Engine/Source/Runtime/Core/Public/Containers/VersePathFwd.h", // invalid forward declaration - 'namespace UE::Core'
 		};
 
 		/// <summary>
