IncludeTool: Detect invalid newlines in string literals, and add support for parsing raw string literals.

diff --git a/Engine/Source/Programs/IncludeTool/IncludeTool/Preprocessor.cs b/Engine/Source/Programs/IncludeTool/IncludeTool/Preprocessor.cs
--- a/Engine/Source/Programs/IncludeTool/IncludeTool/Preprocessor.cs
+++ b/Engine/Source/Programs/IncludeTool/IncludeTool/Preprocessor.cs
@@ -301,7 +301,7 @@ namespace IncludeTool
 			if(PreludeFile != null)
 			{
 				TextBuffer PreludeText = TextBuffer.FromFile(PreludeFile.FullName);
-				PreprocessorMarkup[] PreludeMarkup = PreprocessorMarkup.ParseArray(PreludeText);
+				PreprocessorMarkup[] PreludeMarkup = PreprocessorMarkup.ParseArray("prelude", PreludeText);
 				foreach(PreprocessorMarkup Markup in PreludeMarkup)
 				{
 					ParseMarkup(Markup.Type, Markup.Tokens, Markup.Location.LineIdx);
@@ -311,7 +311,7 @@ namespace IncludeTool
 			foreach(string Definition in Definitions)
 			{
 				// Create a token reader for this definition
-				TokenReader Reader = new TokenReader(TextBuffer.FromString(Definition), TextLocation.Origin);
+				TokenReader Reader = new TokenReader($"(predefined value for {Definition})", TextBuffer.FromString(Definition), TextLocation.Origin);
 
 				// Read it into a buffer
 				List<Token> Tokens = new List<Token>();
@@ -1511,7 +1511,7 @@ namespace IncludeTool
 						// Write the expanded macro tokens if the current block is active, or blank lines if not
 						if (IsBranchActive())
 						{
-							WriteExpandedText(SourceFile.Text, Location, Markup.EndLocation, Writer);
+							WriteExpandedText(SourceFile.Location.FullName, SourceFile.Text, Location, Markup.EndLocation, Writer);
 						}
 						else
 						{
@@ -1711,16 +1711,17 @@ namespace IncludeTool
 		/// <summary>
 		/// Write a portion of the given file to the output stream, expanding macros in it
 		/// </summary>
+		/// <param name="FileName">The file being parsed</param>
 		/// <param name="Text">The text buffer for the current file</param>
 		/// <param name="Location">Start of the region to copy to the output stream</param>
 		/// <param name="EndLocation">End of the region to copy to the output stream</param>
 		/// <param name="Writer">Writer for the output text</param>
-		void WriteExpandedText(TextBuffer Text, TextLocation Location, TextLocation EndLocation, TextWriter Writer)
+		void WriteExpandedText(string FileName, TextBuffer Text, TextLocation Location, TextLocation EndLocation, TextWriter Writer)
 		{
 			int ExpectedLineIdx = Location.LineIdx;
 
 			// Expand any macros in this block of tokens
-			TokenReader Reader = new TokenReader(Text, Location);
+			TokenReader Reader = new TokenReader(FileName, Text, Location);
 			for(bool bMoveNext = Reader.MoveNext(); bMoveNext; )
 			{
 				// Write the whitespace to the output stream
@@ -1756,7 +1757,7 @@ namespace IncludeTool
 					// If we expanded a macro, insert all the comment tokens
 					if(Reader.TokenWhitespaceLocation > FirstArgumentLocation)
 					{
-						TokenReader CommentReader = new TokenReader(Text, FirstArgumentLocation);
+						TokenReader CommentReader = new TokenReader(FileName, Text, FirstArgumentLocation);
 						while(CommentReader.MoveNext() && CommentReader.TokenLocation < Reader.TokenWhitespaceLocation)
 						{
 							string[] Comments = Text.Extract(CommentReader.TokenWhitespaceLocation, CommentReader.TokenLocation).Select(x => x.EndsWith("\\")? x.Substring(0, x.Length - 1) : x).ToArray();
@@ -1973,7 +1974,7 @@ namespace IncludeTool
 						{
 							List<Token> Tokens = new List<Token>();
 
-							TokenReader Reader = new TokenReader(File.Text, Markup.Location);
+							TokenReader Reader = new TokenReader("", File.Text, Markup.Location);
 							while(Reader.MoveNext() && Reader.TokenLocation < Markup.EndLocation)
 							{
 								Tokens.Add(Reader.Current);
diff --git a/Engine/Source/Programs/IncludeTool/IncludeTool/PreprocessorExpression.cs b/Engine/Source/Programs/IncludeTool/IncludeTool/PreprocessorExpression.cs
--- a/Engine/Source/Programs/IncludeTool/IncludeTool/PreprocessorExpression.cs
+++ b/Engine/Source/Programs/IncludeTool/IncludeTool/PreprocessorExpression.cs
@@ -591,7 +591,7 @@ namespace IncludeTool
 		/// <param name="ExpectedResult">The expected value of the expression</param>
         static void RunTest(string Expression, long? ExpectedResult)
         {
-			TokenReader Reader = new TokenReader(Expression);
+			TokenReader Reader = new TokenReader("", Expression);
 
 			List<Token> Tokens = new List<Token>();
 			while (Reader.MoveNext() && Reader.Current.Text != "\n")
diff --git a/Engine/Source/Programs/IncludeTool/IncludeTool/PreprocessorMarkup.cs b/Engine/Source/Programs/IncludeTool/IncludeTool/PreprocessorMarkup.cs
--- a/Engine/Source/Programs/IncludeTool/IncludeTool/PreprocessorMarkup.cs
+++ b/Engine/Source/Programs/IncludeTool/IncludeTool/PreprocessorMarkup.cs
@@ -252,9 +252,9 @@ namespace IncludeTool
 		/// </summary>
 		/// <param name="Reader">Reader for token objects</param>
 		/// <returns>Array of markup objects which split up the given text buffer</returns>
-		public static PreprocessorMarkup[] ParseArray(TextBuffer Text)
+		public static PreprocessorMarkup[] ParseArray(string FileName, TextBuffer Text)
 		{
-			TokenReader Reader = new TokenReader(Text, TextLocation.Origin);
+			TokenReader Reader = new TokenReader(FileName, Text, TextLocation.Origin);
 
 			List<PreprocessorMarkup> Markup = new List<PreprocessorMarkup>();
 			if(Reader.MoveNext())
@@ -358,7 +358,7 @@ namespace IncludeTool
 						}
 
 						// Create the new fragment
-						PreprocessorMarkupType Type = IsIncludeMarkup(Text, StartLocation, Reader.TokenLocation)? PreprocessorMarkupType.IncludeMarkup : PreprocessorMarkupType.Text;
+						PreprocessorMarkupType Type = IsIncludeMarkup(FileName, Text, StartLocation, Reader.TokenLocation)? PreprocessorMarkupType.IncludeMarkup : PreprocessorMarkupType.Text;
 						Markup.Add(new PreprocessorMarkup(Type, StartLocation, Reader.TokenLocation, null));
 					}
 					else
@@ -374,13 +374,14 @@ namespace IncludeTool
 		/// <summary>
 		/// Checks whether the given block of text is an include decoration
 		/// </summary>
+		/// <param name="FileName">Name of the file being parsed</param>
 		/// <param name="Text">The text buffer to read from</param>
 		/// <param name="StartLocation">Start of the block of text to check</param>
 		/// <param name="EndLocation">End of the block of text to check</param>
 		/// <returns>True if the region just consists of #include markup macros</returns>
-		public static bool IsIncludeMarkup(TextBuffer Text, TextLocation StartLocation, TextLocation EndLocation)
+		public static bool IsIncludeMarkup(string FileName, TextBuffer Text, TextLocation StartLocation, TextLocation EndLocation)
 		{
-			TokenReader Reader = new TokenReader(Text, StartLocation, EndLocation);
+			TokenReader Reader = new TokenReader(FileName, Text, StartLocation, EndLocation);
 			while(Reader.MoveNext(TokenReaderContext.IgnoreNewlines))
 			{
 				if(Reader.Current.Text == "MONOLITHIC_HEADER_BOILERPLATE" && Reader.MoveNext() && Reader.Current.Text == "(" && Reader.MoveNext() && Reader.Current.Text == ")")
diff --git a/Engine/Source/Programs/IncludeTool/IncludeTool/Program.cs b/Engine/Source/Programs/IncludeTool/IncludeTool/Program.cs
--- a/Engine/Source/Programs/IncludeTool/IncludeTool/Program.cs
+++ b/Engine/Source/Programs/IncludeTool/IncludeTool/Program.cs
@@ -968,26 +968,37 @@ namespace IncludeTool
 		/// Preprocess a given file, updating the markup for each included file
 		/// </summary>
 		/// <param name="Location">The file to process</param>
-		/// <param name="PreludeDefinitions">Built-in definitions for the preprocessor, scraped from the underlying compiler</param>
+		/// <param name="PreludeLocation">Built-in definitions for the preprocessor, scraped from the underlying compiler</param>
 		/// <param name="CompileEnvironment">Compile environment for the given file</param>
 		/// <param name="FileToActiveMarkup">Map from source file to bit array of the markup blocks that are active. Used to verify all preprocessed files parse the same way.</param>
 		/// <param name="Log">Writer for output messages</param>
 		static void PreprocessFile(FileReference Location, FileReference PreludeLocation, CompileEnvironment CompileEnvironment, IEnumerable<DirectoryReference> ExtraSystemIncludePaths, ConcurrentDictionary<SourceFile, MarkupState> FileToActiveMarkup, LineBasedTextWriter Log)
 		{
-			// Get the initial file
-			PreprocessorFile InitialFile = new PreprocessorFile(Location.FullName, Workspace.GetFile(Location));
+			try
+			{
+				// Get the initial file
+				PreprocessorFile InitialFile = new PreprocessorFile(Location.FullName, Workspace.GetFile(Location));
 
-			// Get all the include paths. For Clang, we have to add these paths manually.
-			List<DirectoryReference> IncludePaths = new List<DirectoryReference>(CompileEnvironment.IncludePaths);
-			IncludePaths.AddRange(ExtraSystemIncludePaths);
+				// Get all the include paths. For Clang, we have to add these paths manually.
+				List<DirectoryReference> IncludePaths = new List<DirectoryReference>(CompileEnvironment.IncludePaths);
+				IncludePaths.AddRange(ExtraSystemIncludePaths);
 
-			// Set up the preprocessor with the correct defines for this environment
-			Preprocessor PreprocessorInst = new Preprocessor(PreludeLocation, CompileEnvironment.Definitions, IncludePaths);
-			foreach(FileReference ForceIncludedFile in CompileEnvironment.ForceIncludeFiles)
+				// Set up the preprocessor with the correct defines for this environment
+				Preprocessor PreprocessorInst = new Preprocessor(PreludeLocation, CompileEnvironment.Definitions, IncludePaths);
+				foreach (FileReference ForceIncludedFile in CompileEnvironment.ForceIncludeFiles)
+				{
+					PreprocessIncludedFile(PreprocessorInst, new PreprocessorFile(ForceIncludedFile.FullName, Workspace.GetFile(ForceIncludedFile)), FileToActiveMarkup, Log);
+				}
+				PreprocessIncludedFile(PreprocessorInst, InitialFile, FileToActiveMarkup, Log);
+			}
+			catch (TokenException ex)
+			{
+				Log.WriteLine(ex.Message);
+			}
+			catch (Exception ex)
 			{
-				PreprocessIncludedFile(PreprocessorInst, new PreprocessorFile(ForceIncludedFile.FullName, Workspace.GetFile(ForceIncludedFile)), FileToActiveMarkup, Log);
+				Log.WriteLine("Error parsing {0}: {1}", Location, ex.ToString());
 			}
-			PreprocessIncludedFile(PreprocessorInst, InitialFile, FileToActiveMarkup, Log);
 		}
 
 		/// <summary>
diff --git a/Engine/Source/Programs/IncludeTool/IncludeTool/SourceFile.cs b/Engine/Source/Programs/IncludeTool/IncludeTool/SourceFile.cs
--- a/Engine/Source/Programs/IncludeTool/IncludeTool/SourceFile.cs
+++ b/Engine/Source/Programs/IncludeTool/IncludeTool/SourceFile.cs
@@ -234,7 +234,7 @@ namespace IncludeTool
 			}
 			else
 			{
-				Markup = PreprocessorMarkup.ParseArray(Text);
+				Markup = PreprocessorMarkup.ParseArray(Location?.FullName ?? "unknown", Text);
 			}
 
 			// Find the markup range which excludes header guards
@@ -492,7 +492,7 @@ namespace IncludeTool
 						if(TrimLine.Length > 0 && !TrimLine.Equals("// Forward declarations", StringComparison.OrdinalIgnoreCase) && !TrimLine.Equals("// Forward declarations.", StringComparison.OrdinalIgnoreCase))
 						{
 							// Create a token reader for the current line
-							TokenReader Reader = new TokenReader(Text, new TextLocation(NewLastLocation.LineIdx, 0), new TextLocation(NewLastLocation.LineIdx, Text.Lines[NewLastLocation.LineIdx].Length));
+							TokenReader Reader = new TokenReader(Location.FullName, Text, new TextLocation(NewLastLocation.LineIdx, 0), new TextLocation(NewLastLocation.LineIdx, Text.Lines[NewLastLocation.LineIdx].Length));
 
 							// Read it into a buffer
 							List<Token> Tokens = new List<Token>();
@@ -646,12 +646,13 @@ namespace IncludeTool
 		/// <summary>
 		/// Split the given line into tokens
 		/// </summary>
+		/// <param name="FileName">File being parsed</param>
 		/// <param name="Text">Buffer to read from</param>
 		/// <param name="LineIdx">Index of the line to tokenize</param>
 		/// <returns>Sequence of tokens</returns>
-		static IEnumerable<Token> TokenizeLine(TextBuffer Text, int LineIdx)
+		static IEnumerable<Token> TokenizeLine(string FileName, TextBuffer Text, int LineIdx)
 		{
-			TokenReader Reader = new TokenReader(Text, new TextLocation(LineIdx, 0), new TextLocation(LineIdx, Text.Lines[LineIdx].Length));
+			TokenReader Reader = new TokenReader(FileName, Text, new TextLocation(LineIdx, 0), new TextLocation(LineIdx, Text.Lines[LineIdx].Length));
 			while(Reader.MoveNext())
 			{
 				yield return Reader.Current;
diff --git a/Engine/Source/Programs/IncludeTool/IncludeTool/SymbolTable.cs b/Engine/Source/Programs/IncludeTool/IncludeTool/SymbolTable.cs
--- a/Engine/Source/Programs/IncludeTool/IncludeTool/SymbolTable.cs
+++ b/Engine/Source/Programs/IncludeTool/IncludeTool/SymbolTable.cs
@@ -294,7 +294,7 @@ namespace IncludeTool
 		/// <param name="Scope">The scope depth; the number of unmatched opening brace tokens</param>
 		void ParseDeclarations(SourceFile File, SourceFragment Fragment, PreprocessorMarkup Markup, ref int Scope)
 		{
-			TokenReader Reader = new TokenReader(File.Text, Markup.Location, Markup.EndLocation);
+			TokenReader Reader = new TokenReader(File.Location.FullName, File.Text, Markup.Location, Markup.EndLocation);
 			if(Reader.MoveNext())
 			{
 				for(;;)
@@ -723,7 +723,7 @@ namespace IncludeTool
 		{
 			foreach (PreprocessorMarkup Markup in HeaderFile.Markup.Where(x => x.Type == PreprocessorMarkupType.Text))
 			{
-				TokenReader Reader = new TokenReader(HeaderFile.Text, Markup.Location, Markup.EndLocation);
+				TokenReader Reader = new TokenReader(HeaderFile.Location.FullName, HeaderFile.Text, Markup.Location, Markup.EndLocation);
 				while(Reader.MoveNext(TokenReaderContext.IgnoreNewlines))
 				{
 					if(!ReadClassOrStructForwardDeclaration(Reader, HeaderFile, SymbolToHeader, Log) 
@@ -994,7 +994,7 @@ namespace IncludeTool
 			TextBuffer Text = Fragment.File.Text;
 			if(Text != null && Fragment.MarkupMax > Fragment.MarkupMin)
 			{
-				TokenReader Reader = new TokenReader(Fragment.File.Text, Fragment.MinLocation, Fragment.MaxLocation);
+				TokenReader Reader = new TokenReader(Fragment.File.Location.FullName, Fragment.File.Text, Fragment.MinLocation, Fragment.MaxLocation);
 				for(bool bMoveNext = Reader.MoveNext(TokenReaderContext.IgnoreNewlines); bMoveNext; )
 				{
 					// Read the current token, and immediately move to the next so that we can lookahead if we need to
diff --git a/Engine/Source/Programs/IncludeTool/IncludeTool/Token.cs b/Engine/Source/Programs/IncludeTool/IncludeTool/Token.cs
--- a/Engine/Source/Programs/IncludeTool/IncludeTool/Token.cs
+++ b/Engine/Source/Programs/IncludeTool/IncludeTool/Token.cs
@@ -1,4 +1,4 @@
-﻿// Copyright Epic Games, Inc. All Rights Reserved.
+// Copyright Epic Games, Inc. All Rights Reserved.
 
 using System;
 using System.Collections.Generic;
@@ -157,7 +157,7 @@ namespace IncludeTool
 		{
 			string Text = FirstToken.Text + SecondToken.Text;
 
-			TokenReader Reader = new TokenReader(Text);
+			TokenReader Reader = new TokenReader("concatenated", Text);
 			if(Reader.MoveNext())
 			{
 				TokenType Type = Reader.Current.Type;
diff --git a/Engine/Source/Programs/IncludeTool/IncludeTool/TokenReader.cs b/Engine/Source/Programs/IncludeTool/IncludeTool/TokenReader.cs
--- a/Engine/Source/Programs/IncludeTool/IncludeTool/TokenReader.cs
+++ b/Engine/Source/Programs/IncludeTool/IncludeTool/TokenReader.cs
@@ -37,11 +37,32 @@ namespace IncludeTool
 		IgnoreNewlines = 0x04,
 	}
 
+	/// <summary>
+	/// Exception thrown when an invalid token is encountered
+	/// </summary>
+	class TokenException : Exception
+	{
+		public string FileName { get; }
+		public TextLocation Location { get; }
+		public string Description { get; }
+
+		public TokenException(string FileName, TextLocation Location, string Description)
+			: base($"{FileName}({Location.LineIdx + 1}): error: {Description}")
+		{
+			this.Description = Description;
+		}
+	}
+
 	/// <summary>
 	/// Tokenizer for C++ source files. Provides functions for navigating a source file skipping whitespace, comments, and so on when required.
 	/// </summary>
 	class TokenReader : IEnumerator<Token>
 	{
+		/// <summary>
+		/// The current file
+		/// </summary>
+		public string FileName { get; }
+
 		/// <summary>
 		/// The current text buffer being read from
 		/// </summary>
@@ -61,8 +82,8 @@ namespace IncludeTool
 		/// Constructor
 		/// </summary>
 		/// <param name="Text">The text to tokenize</param>
-		public TokenReader(string Text)
-			: this(TextBuffer.FromString(Text), TextLocation.Origin)
+		public TokenReader(string FileName, string Text)
+			: this(FileName, TextBuffer.FromString(Text), TextLocation.Origin)
 		{
 		}
 
@@ -71,8 +92,8 @@ namespace IncludeTool
 		/// </summary>
 		/// <param name="Text">The text buffer to tokenize</param>
 		/// <param name="Location">Initial location to start reading from</param>
-		public TokenReader(TextBuffer Text, TextLocation Location)
-			: this(Text, Location, Text.End)
+		public TokenReader(string FileName, TextBuffer Text, TextLocation Location)
+			: this(FileName, Text, Location, Text.End)
 		{
 		}
 
@@ -82,8 +103,9 @@ namespace IncludeTool
 		/// <param name="Text">The text buffer to tokenize</param>
 		/// <param name="Location">Initial location to start reading from</param>
 		/// <param name="EndLocation">The end location to read from</param>
-		public TokenReader(TextBuffer Text, TextLocation Location, TextLocation EndLocation)
+		public TokenReader(string FileName, TextBuffer Text, TextLocation Location, TextLocation EndLocation)
 		{
+			this.FileName = FileName;
 			this.Text = Text;
 			this.TokenWhitespaceLocation = Location;
 			this.TokenLocation = Location;
@@ -203,7 +225,7 @@ namespace IncludeTool
 
 			// Read the token
 			TokenLocation = new TextLocation(LineIdx, ColumnIdx);
-			bool bResult = ReadToken(Text, ref LineIdx, ref ColumnIdx, bHasWhitespace? TokenFlags.HasLeadingSpace : TokenFlags.None, Context, out CurrentToken);
+			bool bResult = ReadToken(FileName, Text, ref LineIdx, ref ColumnIdx, bHasWhitespace? TokenFlags.HasLeadingSpace : TokenFlags.None, Context, out CurrentToken);
 			TokenEndLocation = new TextLocation(LineIdx, ColumnIdx);
 			return bResult;
 		}
@@ -315,12 +337,13 @@ namespace IncludeTool
 		/// <summary>
 		/// Reads a single token from a text buffer
 		/// </summary>
+		/// <param name="FileName">File being parsed</param>
 		/// <param name="Text">The text buffer to read from</param>
 		/// <param name="LineIdx">The current line index</param>
 		/// <param name="ColumnIdx">The current column index</param>
 		/// <param name="Flags">Flags for the new token</param>
 		/// <returns>The next token, or null at the end of the file</returns>
-		static bool ReadToken(TextBuffer Text, ref int LineIdx, ref int ColumnIdx, TokenFlags Flags, TokenReaderContext Context, out Token Result)
+		static bool ReadToken(string FileName, TextBuffer Text, ref int LineIdx, ref int ColumnIdx, TokenFlags Flags, TokenReaderContext Context, out Token Result)
 		{
 			int StartLineIdx = LineIdx;
 			int StartColumnIdx = ColumnIdx;
@@ -346,7 +369,7 @@ namespace IncludeTool
 					for(;;)
 					{
 						Character = Text[LineIdx, ColumnIdx];
-						if(Character == '\n')
+						if(Character == '\n' && Builder[^1] != '\\')
 						{
 							break;
 						}
@@ -363,21 +386,24 @@ namespace IncludeTool
 			else if(Character == '\'')
 			{
 				// Character literal
-				SkipTextLiteral(Text, ref LineIdx, ref ColumnIdx, '\'');
+				SkipTextLiteral(FileName, Text, ref LineIdx, ref ColumnIdx, '\'');
 				Result = new Token(Text.ExtractString(StartLineIdx, StartColumnIdx, LineIdx, ColumnIdx), TokenType.CharacterLiteral, Flags);
 				return true;
 			}
 			else if(Character == '\"')
 			{
 				// String literal
-				SkipTextLiteral(Text, ref LineIdx, ref ColumnIdx, '\"');
+				SkipTextLiteral(FileName, Text, ref LineIdx, ref ColumnIdx, '\"');
 				Result = new Token(Text.ExtractString(StartLineIdx, StartColumnIdx, LineIdx, ColumnIdx), TokenType.StringLiteral, Flags);
 				return true;
 			}
 			else if ((Character >= 'a' && Character <= 'z') || (Character >= 'A' && Character <= 'Z') || Character == '_')
 			{
+				char FirstCharacter = Character;
+
 				// Identifier (or text literal with prefix)
-				for(;;)
+				int PrefixLength = 1;
+				for (;; PrefixLength++)
 				{
 					Character = Text[LineIdx, ColumnIdx];
 					if((Character < 'a' || Character > 'z') && (Character < 'A' || Character > 'Z') && (Character < '0' || Character > '9') && Character != '_' && Character != '$')
@@ -388,17 +414,23 @@ namespace IncludeTool
 				}
 
 				// Check if it's a prefixed text literal
-				if(Character == '\'')
+				if (Character == '\'')
 				{
 					Text.MoveNext(ref LineIdx, ref ColumnIdx);
-					SkipTextLiteral(Text, ref LineIdx, ref ColumnIdx, '\'');
+					SkipTextLiteral(FileName, Text, ref LineIdx, ref ColumnIdx, '\'');
 					Result = new Token(Text.ExtractString(StartLineIdx, StartColumnIdx, LineIdx, ColumnIdx), TokenType.CharacterLiteral, Flags);
 					return true;
 				}
-				else if(Character == '\"')
+				else if (Character == '\"' && PrefixLength == 1 && FirstCharacter == 'R')
+				{
+					SkipRawStringLiteral(FileName, Text, ref LineIdx, ref ColumnIdx);
+					Result = new Token(Text.ExtractString(StartLineIdx, StartColumnIdx, LineIdx, ColumnIdx), TokenType.StringLiteral, Flags);
+					return true;
+				}
+				else if (Character == '\"')
 				{
 					Text.MoveNext(ref LineIdx, ref ColumnIdx);
-					SkipTextLiteral(Text, ref LineIdx, ref ColumnIdx, '\"');
+					SkipTextLiteral(FileName, Text, ref LineIdx, ref ColumnIdx, '\"');
 					Result = new Token(Text.ExtractString(StartLineIdx, StartColumnIdx, LineIdx, ColumnIdx), TokenType.StringLiteral, Flags);
 					return true;
 				}
@@ -483,28 +515,77 @@ namespace IncludeTool
 		/// <summary>
 		/// Skip past a text literal (a quoted character literal or string literal)
 		/// </summary>
+		/// <param name="FileName">The file being parsed</param>
 		/// <param name="Text">The text buffer to read from</param>
 		/// <param name="LineIdx">The current line index</param>
 		/// <param name="ColumnIdx">The current column index</param>
 		/// <param name="LastCharacter">The terminating character to look for, ignoring escape sequences</param>
-		static void SkipTextLiteral(TextBuffer Text, ref int LineIdx, ref int ColumnIdx, char LastCharacter)
+		static void SkipTextLiteral(string FileName, TextBuffer Text, ref int LineIdx, ref int ColumnIdx, char LastCharacter)
 		{
 			for(;;)
 			{
 				char Character = Text.ReadCharacter(ref LineIdx, ref ColumnIdx);
-				if(Character == '\0')
+				if (Character == '\0')
 				{
-					throw new Exception("Unexpected end of file in text literal");
+					throw new TokenException(FileName, new TextLocation(LineIdx, ColumnIdx), "Unexpected end of file in text literal");
 				}
-				else if(Character == '\\')
+				else if (Character == '\\')
 				{
 					Text.MoveNext(ref LineIdx, ref ColumnIdx);
 				}
-				else if(Character == LastCharacter)
+				else if (Character == LastCharacter)
+				{
+					break;
+				}
+				else if (Character == '\n')
+				{
+					throw new TokenException(FileName, new TextLocation(LineIdx, ColumnIdx), "Unexpected newline in text literal");
+				}
+			}
+		}
+
+		/// <summary>
+		/// Skip past a text literal (a quoted character literal or string literal)
+		/// </summary>
+		/// <param name="FileName">The file being parsed</param>
+		/// <param name="Text">The text buffer to read from</param>
+		/// <param name="LineIdx">The current line index</param>
+		/// <param name="ColumnIdx">The current column index</param>
+		static void SkipRawStringLiteral(string FileName, TextBuffer Text, ref int LineIdx, ref int ColumnIdx)
+		{
+			int StartColumnIdx = ColumnIdx;
+
+			ReadOnlySpan<char> Line = Text[LineIdx].AsSpan();
+			for (; ; ColumnIdx++)
+			{
+				if (ColumnIdx == Line.Length)
+				{
+					throw new TokenException(FileName, new TextLocation(LineIdx, ColumnIdx), "End of line when looking for raw text delimiter");
+				}
+				else if (Line[ColumnIdx] == '(')
 				{
 					break;
 				}
 			}
+
+			ReadOnlySpan<char> Delim = Line.Slice(StartColumnIdx + 1, ColumnIdx - (StartColumnIdx + 1));
+			while(LineIdx < Text.Lines.Length)
+			{
+				Line = Text[LineIdx].AsSpan();
+
+				for (; ColumnIdx < Line.Length; ColumnIdx++)
+				{
+					if (Line[ColumnIdx] == ')' && Line.Length >= ColumnIdx + 2 + Delim.Length && Line[ColumnIdx + 1 + Delim.Length] == '\"' && Line.Slice(ColumnIdx + 1).StartsWith(Delim))
+					{
+						ColumnIdx += Delim.Length + 2;
+						return;
+					}
+				}
+
+				(LineIdx, ColumnIdx) = (LineIdx + 1, 0);
+			}
+
+			throw new TokenException(FileName, new TextLocation(LineIdx, ColumnIdx), "End of file while looking for end of raw string delimiter");
 		}
 	}
 }
