[IncludeTool]
* Added TaskGraphFwd.h to exclusion list to remove warning about fwd file not only containing includes

diff --git a/Engine/Source/Programs/IncludeTool/IncludeTool/Rules.cs b/Engine/Source/Programs/IncludeTool/IncludeTool/Rules.cs
--- a/Engine/Source/Programs/IncludeTool/IncludeTool/Rules.cs
+++ b/Engine/Source/Programs/IncludeTool/IncludeTool/Rules.cs
@@ -424,7 +424,8 @@ namespace IncludeTool
 			"/Engine/Source/Runtime/Slate/Public/Framework/Text/ShapedTextCacheFwd.h", // Typedef isn't a forward declaration
 			"/Engine/Source/Runtime/MovieScene/Public/MovieSceneFwd.h",
             "/Engine/Source/Editor/SequencerCore/Public/SequencerCoreFwd.h", // invalid forward declaration - 'namespace UE'
-            "/Engine/Source/Runtime/Core/Public/Math/MathFwd.h", // invalid forward declaration - 'namespace UE::Math'
+			"/Engine/Source/Runtime/Core/Public/Async/TaskGraphFwd.h",  // warning: expected only include directives and text in forward declaration header in TaskGraphFwd.h
+			"/Engine/Source/Runtime/Core/Public/Math/MathFwd.h", // invalid forward declaration - 'namespace UE::Math'
 			"/Engine/Source/Runtime/Core/Public/Containers/ContainersFwd.h", // invalid forward declaration - 'template<> struct TIsContiguousContainer<Type> { static constexpr bool Value = true; };'
 			"/Engine/Source/Runtime/Core/Public/Containers/StringFwd.h", // invalid forward declaration - 'template<> struct TIsContiguousContainer<Type> { static constexpr bool Value = true; };'
 			"/Engine/Source/Runtime/Core/Public/Internationalization/StringTableCoreFwd.h", // Typedef isn't a forward declaration
