Add HttpFwd and OnlineSubsystemMcpFwd to exclude list. They include namespaces and typedefs.

diff --git a/Engine/Source/Programs/IncludeTool/IncludeTool/Rules.cs b/Engine/Source/Programs/IncludeTool/IncludeTool/Rules.cs
--- a/Engine/Source/Programs/IncludeTool/IncludeTool/Rules.cs
+++ b/Engine/Source/Programs/IncludeTool/IncludeTool/Rules.cs
@@ -441,6 +441,8 @@ namespace IncludeTool
 			"/Engine/Source/Runtime/Core/Public/Containers/VersePathFwd.h", // invalid forward declaration - 'namespace UE::Core'
 			"/Engine/Source/Runtime/Experimental/Iris/Core/Public/Iris/ReplicationState/ReplicationStateFwd.h", // invalid forward declaration - 'namespace UE::Net'
 			"/Engine/Source/Runtime/Experimental/Iris/Core/Private/Iris/ReplicationSystem/ObjectReferenceCacheFwd.h", // invalid forward declaration - 'namespace UE::Net'
+			"/Engine/Source/Runtime/Online/HTTP/Public/HttpFwd.h", // error: invalid forward declaration - 'typedef TSharedPtr<IHttpRequest, ESPMode::ThreadSafe> FHttpRequestPtr;'
+			"/Engine/Restricted/NotForLicensees/Plugins/Online/OnlineSubsystemMcp/Source/Public/OnlineSubsystemMcpFwd.h", // expected only include directives and text in forward declaration header
 		};
 
 		/// <summary>
