Adding RHIFwd.h to IncludeTool fwd header exclude list.

diff --git a/Engine/Source/Programs/IncludeTool/IncludeTool/Rules.cs b/Engine/Source/Programs/IncludeTool/IncludeTool/Rules.cs
--- a/Engine/Source/Programs/IncludeTool/IncludeTool/Rules.cs
+++ b/Engine/Source/Programs/IncludeTool/IncludeTool/Rules.cs
@@ -443,6 +443,7 @@ namespace IncludeTool
 			"/Engine/Source/Runtime/Experimental/Iris/Core/Private/Iris/ReplicationSystem/ObjectReferenceCacheFwd.h", // invalid forward declaration - 'namespace UE::Net'
 			"/Engine/Source/Runtime/Online/HTTP/Public/HttpFwd.h", // error: invalid forward declaration - 'typedef TSharedPtr<IHttpRequest, ESPMode::ThreadSafe> FHttpRequestPtr;'
 			"/Engine/Restricted/NotForLicensees/Plugins/Online/OnlineSubsystemMcp/Source/Public/OnlineSubsystemMcpFwd.h", // expected only include directives and text in forward declaration header
+			"/Engine/Source/Runtime/RHI/Public/RHIFwd.h", // invalid forward declaration - 'namespace ERHIFeatureLevel { enum Type : int; }'
 		};
 
 		/// <summary>
