IncludeTool: Use relative paths for printing out conflicting symbol declarations.

diff --git a/Engine/Source/Programs/IncludeTool/IncludeTool/Program.cs b/Engine/Source/Programs/IncludeTool/IncludeTool/Program.cs
--- a/Engine/Source/Programs/IncludeTool/IncludeTool/Program.cs
+++ b/Engine/Source/Programs/IncludeTool/IncludeTool/Program.cs
@@ -658,7 +658,7 @@ namespace IncludeTool
 					}
 
 					// Print all the conflicting declarations
-					SymbolTable.PrintConflicts(Log);
+					SymbolTable.PrintConflicts(InputDir, Log);
 				}
 
 				// Build a list of forward declarations per file
diff --git a/Engine/Source/Programs/IncludeTool/IncludeTool/SymbolTable.cs b/Engine/Source/Programs/IncludeTool/IncludeTool/SymbolTable.cs
--- a/Engine/Source/Programs/IncludeTool/IncludeTool/SymbolTable.cs
+++ b/Engine/Source/Programs/IncludeTool/IncludeTool/SymbolTable.cs
@@ -1,5 +1,6 @@
 // Copyright Epic Games, Inc. All Rights Reserved.
 
+using EpicGames.Core;
 using IncludeTool.Support;
 using System;
 using System.Collections.Generic;
@@ -165,8 +166,9 @@ namespace IncludeTool
 		/// <summary>
 		/// Prints all the symbols with conflicting definitions
 		/// </summary>
+		/// <param name="InputDir"></param>
 		/// <param name="Log">Writer for log messagsThe source fragment to parse</param>
-		public void PrintConflicts(LineBasedTextWriter Log)
+		public void PrintConflicts(DirectoryReference InputDir, LineBasedTextWriter Log)
 		{
 			foreach(string SymbolName in Lookup.Keys)
 			{
@@ -191,7 +193,8 @@ namespace IncludeTool
 						Log.WriteLine("warning: conflicting declarations of '{0}':", SymbolName);
 						foreach(Symbol Symbol in Symbols)
 						{
-							Log.WriteLine($"  {Symbol.Type} at {Symbol.Fragment.File}:{Symbol.Location.LineIdx+1}");
+							string RelativePath = Symbol.Fragment.File.Location.MakeRelativeTo(InputDir);
+							Log.WriteLine($"  {Symbol.Type} at {RelativePath}:{Symbol.Location.LineIdx+1}");
 						}
 					}
 				}
