IncludeTool: Suppress conflicting declarations of 'FGraphEvent' warning

#rnx

diff --git a/Engine/Source/Programs/IncludeTool/IncludeTool/SymbolTable.cs b/Engine/Source/Programs/IncludeTool/IncludeTool/SymbolTable.cs
--- a/Engine/Source/Programs/IncludeTool/IncludeTool/SymbolTable.cs
+++ b/Engine/Source/Programs/IncludeTool/IncludeTool/SymbolTable.cs
@@ -189,6 +189,10 @@ namespace IncludeTool
 						{
 							continue;
 						}
+						if (FileName.Equals("TaskGraphFwd.h", StringComparison.OrdinalIgnoreCase))
+						{
+							continue;
+						}
 
 						Log.WriteLine("warning: conflicting declarations of '{0}':", SymbolName);
 						foreach(Symbol Symbol in Symbols)
