[IncludeTool]
* Excluded OptionalFwd.h from fwd check

diff --git a/Engine/Source/Programs/IncludeTool/IncludeTool/Rules.cs b/Engine/Source/Programs/IncludeTool/IncludeTool/Rules.cs
--- a/Engine/Source/Programs/IncludeTool/IncludeTool/Rules.cs
+++ b/Engine/Source/Programs/IncludeTool/IncludeTool/Rules.cs
@@ -446,7 +446,8 @@ namespace IncludeTool
 			"/Engine/Source/Runtime/RHI/Public/RHIFwd.h", // invalid forward declaration - 'namespace ERHIFeatureLevel { enum Type : int; }'
 			"/Engine/Source/Runtime/RenderCore/Public/RenderGraphFwd.h",
 			"/Engine/Source/Runtime/Core/Public/Templates/SharedPointerFwd.h", // Has an enum as well
-
+			"/Engine/Source/Runtime/Core/Public/Templates/SharedPointerFwd.h", // Has an enum as well
+			"/Engine/Source/Runtime/Core/Public/Misc/OptionalFwd.h", // Has special struct
 		};
 
 		/// <summary>
