Removed places where -nobuilduht is added to the command line.

diff --git a/Engine/Build/Graph/Tasks/BuildAndTestProject.xml b/Engine/Build/Graph/Tasks/BuildAndTestProject.xml
--- a/Engine/Build/Graph/Tasks/BuildAndTestProject.xml
+++ b/Engine/Build/Graph/Tasks/BuildAndTestProject.xml
@@ -422,7 +422,7 @@
 			<Node Name="$(PlatformCompileNodeName)" Requires="$(PlatformToolsNodeName)" Produces="#$(PlatformCompileNodeName) Binaries">
 				<Switch>
 					<Case If="$(IsContentOnlyProject)">
-						<Command Name="BuildGame" Arguments="$(BCRArgs) -configuration=$(TargetConfigurations) -notools -ubtargs=&quot;-nobuilduht -BuildVersion=$(BuildName) $(GenericCompileArguments) $(ExtraTargetCompileArguments)&quot;" />
+						<Command Name="BuildGame" Arguments="$(BCRArgs) -configuration=$(TargetConfigurations) -notools -ubtargs=&quot;-BuildVersion=$(BuildName) $(GenericCompileArguments) $(ExtraTargetCompileArguments)&quot;" />
 						<Tag Files="$(RootDir)/$(ProjectPath)/Binaries/..." Except="$(ProjectPath)/Binaries/DotNET/..." With="#$(PlatformCompileNodeName) Binaries"/>
 					</Case>
 					<Default>
diff --git a/Engine/Source/Programs/AutomationTool/AutomationUtils/UnrealBuild.cs b/Engine/Source/Programs/AutomationTool/AutomationUtils/UnrealBuild.cs
--- a/Engine/Source/Programs/AutomationTool/AutomationUtils/UnrealBuild.cs
+++ b/Engine/Source/Programs/AutomationTool/AutomationUtils/UnrealBuild.cs
@@ -42,6 +42,7 @@ namespace AutomationTool
 		/// <summary>
 		/// If true we will let UBT build UHT
 		/// </summary>
+		[Obsolete]
 		public bool AlwaysBuildUHT { get; set; }
 
 		public void AddBuildProduct(string InFile)
@@ -123,7 +124,7 @@ namespace AutomationTool
 
 			ClearExportedXGEXML();
 
-			string UHTArg = this.AlwaysBuildUHT ? "" : "-nobuilduht";
+			string UHTArg = "";
 
 			CommandUtils.RunUBT(CommandUtils.CmdEnv, UnrealBuildToolDll: UnrealBuildToolDll, Project: UprojectPath, Target: TargetName, Platform: Platform, Config: Config, AdditionalArgs: String.Format("-Manifest={0} {1} -NoHotReload -xgeexport {2}", CommandUtils.MakePathSafeToUseWithCommandLine(ManifestFile.FullName), UHTArg, AddArgs));
 
@@ -190,10 +191,6 @@ namespace AutomationTool
 			}
 			AddArgs += " -NoUBTMakefiles";
 			AddArgs += " " + InAddArgs;
-			if (!TargetName.Equals("UnrealHeaderTool", StringComparison.InvariantCultureIgnoreCase))
-			{
-				AddArgs += " -nobuilduht";
-			}
 
 			PrepareUBT();
 			using (IScope Scope = GlobalTracer.Instance.BuildSpan("Compile").StartActive())
diff --git a/Engine/Source/Programs/AutomationTool/BuildGraph/Tasks/CompileTask.cs b/Engine/Source/Programs/AutomationTool/BuildGraph/Tasks/CompileTask.cs
--- a/Engine/Source/Programs/AutomationTool/BuildGraph/Tasks/CompileTask.cs
+++ b/Engine/Source/Programs/AutomationTool/BuildGraph/Tasks/CompileTask.cs
@@ -149,7 +149,7 @@ namespace AutomationTool.Tasks
 			bAllowParallelExecutor &= Parameters.AllowParallelExecutor;
 			bAllowAllCores &= Parameters.AllowAllCores;
 
-			UnrealBuild.BuildTarget Target = new UnrealBuild.BuildTarget { TargetName = Parameters.Target, Platform = Parameters.Platform, Config = Parameters.Configuration, UprojectPath = CompileTask.FindProjectFile(), UBTArgs = "-nobuilduht " + (Parameters.Arguments ?? ""), Clean = Parameters.Clean };
+			UnrealBuild.BuildTarget Target = new UnrealBuild.BuildTarget { TargetName = Parameters.Target, Platform = Parameters.Platform, Config = Parameters.Configuration, UprojectPath = CompileTask.FindProjectFile(), UBTArgs = (Parameters.Arguments ?? ""), Clean = Parameters.Clean };
 			if(!String.IsNullOrEmpty(Parameters.Tag))
 			{
 				TargetToTagName.Add(Target, Parameters.Tag);
diff --git a/Engine/Source/Programs/AutomationTool/Scripts/BuildTarget.Automation.cs b/Engine/Source/Programs/AutomationTool/Scripts/BuildTarget.Automation.cs
--- a/Engine/Source/Programs/AutomationTool/Scripts/BuildTarget.Automation.cs
+++ b/Engine/Source/Programs/AutomationTool/Scripts/BuildTarget.Automation.cs
@@ -139,7 +139,6 @@ namespace AutomationTool
 			if (!SingleBuild || (ContainsEditor && !NoTools))
 			{
 				UnrealBuild Build = new UnrealBuild(this);
-				Build.AlwaysBuildUHT = true;
 
 				UnrealBuild.BuildAgenda Agenda = new UnrealBuild.BuildAgenda();
 
diff --git a/Engine/Source/Programs/IncludeTool/IncludeTool/Program.cs b/Engine/Source/Programs/IncludeTool/IncludeTool/Program.cs
--- a/Engine/Source/Programs/IncludeTool/IncludeTool/Program.cs
+++ b/Engine/Source/Programs/IncludeTool/IncludeTool/Program.cs
@@ -905,12 +905,12 @@ namespace IncludeTool
 				}
 			}
 
-			if (Utility.Run(Unreal.DotnetPath, String.Format("\"{0}\" -Mode=JsonExport {1} {2} {3}{4} -disableunity -xgeexport -nobuilduht -nopch -nodebuginfo -define:UE_INCLUDE_TOOL=1 -execcodegenactions -outputfile=\"{5}\"",
+			if (Utility.Run(Unreal.DotnetPath, String.Format("\"{0}\" -Mode=JsonExport {1} {2} {3}{4} -disableunity -xgeexport -nopch -nodebuginfo -define:UE_INCLUDE_TOOL=1 -execcodegenactions -outputfile=\"{5}\"",
 				Unreal.UnrealBuildToolDllPath, Target, Configuration, Platform, Precompile? " -precompile" : "", TaskListFile.ChangeExtension(".json").FullName), Unreal.EngineDirectory, Log) != 0)
 			{
 				throw new Exception("UnrealBuildTool failed");
 			}
-			if (Utility.Run(Unreal.DotnetPath, String.Format("\"{0}\" {1} {2} {3}{4} -disableunity -xgeexport -nobuilduht -nopch -nodebuginfo -define:UE_INCLUDE_TOOL=1",
+			if (Utility.Run(Unreal.DotnetPath, String.Format("\"{0}\" {1} {2} {3}{4} -disableunity -xgeexport -nopch -nodebuginfo -define:UE_INCLUDE_TOOL=1",
 				Unreal.UnrealBuildToolDllPath, Target, Configuration, Platform, Precompile? " -precompile" : ""), Unreal.EngineDirectory, Log) != 0)
 			{
 				throw new Exception("UnrealBuildTool failed");
