Added WeakObjectPtrFwd.h and WeakObjectPtrTemplatesFwd.h to IncludeTool's Fwd-ignore list.

diff --git a/Engine/Source/Programs/IncludeTool/IncludeTool/Rules.cs b/Engine/Source/Programs/IncludeTool/IncludeTool/Rules.cs
--- a/Engine/Source/Programs/IncludeTool/IncludeTool/Rules.cs
+++ b/Engine/Source/Programs/IncludeTool/IncludeTool/Rules.cs
@@ -448,6 +448,8 @@ namespace IncludeTool
 			"/Engine/Source/Runtime/Core/Public/Templates/SharedPointerFwd.h", // Has an enum as well
 			"/Engine/Source/Runtime/Core/Public/Templates/SharedPointerFwd.h", // Has an enum as well
 			"/Engine/Source/Runtime/Core/Public/Misc/OptionalFwd.h", // Has special struct
+			"/Engine/Source/Runtime/CoreUObject/Public/UObject/WeakObjectPtrFwd.h", // error: invalid forward declaration - 'template<> struct TIsPODType<FWeakObjectPtr> { enum { Value = true }; };'
+			"/Engine/Source/Runtime/Core/Public/UObject/WeakObjectPtrTemplatesFwd.h", // error: invalid forward declaration - 'template<class T> struct TIsPODType<TWeakObjectPtr<T> > { enum { Value = true }; };'
 		};
 
 		/// <summary>
