Move DNACalib from MetaHuman stream to Main

diff --git a/Engine/Source/Programs/IncludeTool/IncludeTool/Rules.cs b/Engine/Source/Programs/IncludeTool/IncludeTool/Rules.cs
--- a/Engine/Source/Programs/IncludeTool/IncludeTool/Rules.cs
+++ b/Engine/Source/Programs/IncludeTool/IncludeTool/Rules.cs
@@ -596,11 +596,11 @@ namespace IncludeTool
 				return true;
 			}
 			// Start of RigLogic exclusions
-			if (Markup.Type == PreprocessorMarkupType.Define && (Markup.Tokens[0].Text == "RLAPI" || Markup.Tokens[0].Text == "DNAAPI" || Markup.Tokens[0].Text == "TRIOAPI" || Markup.Tokens[0].Text == "SCAPI" || Markup.Tokens[0].Text == "PMAAPI" || Markup.Tokens[0].Text == "RAFAPI" || Markup.Tokens[0].Text == "GSAPI"))
+			if (Markup.Type == PreprocessorMarkupType.Define && (Markup.Tokens[0].Text == "RLAPI" || Markup.Tokens[0].Text == "DNAAPI" || Markup.Tokens[0].Text == "TRIOAPI" || Markup.Tokens[0].Text == "SCAPI" || Markup.Tokens[0].Text == "PMAAPI" || Markup.Tokens[0].Text == "RAFAPI" || Markup.Tokens[0].Text == "GSAPI" || Markup.Tokens[0].Text == "DNACAPI"))
 			{
 				return true;
 			}
-			if (Markup.Type == PreprocessorMarkupType.Elif && Markup.Tokens.Count == 4 && (Markup.Tokens[2].Text == "RL_SHARED" || Markup.Tokens[2].Text == "GS_SHARED"))
+			if (Markup.Type == PreprocessorMarkupType.Elif && Markup.Tokens.Count == 4 && (Markup.Tokens[2].Text == "RL_SHARED" || Markup.Tokens[2].Text == "GS_SHARED" || Markup.Tokens[2].Text == "DNAC_SHARED"))
 			{
 				return true;
 			}
