IncludeTool: Programmatically control which files should be ignored for symbol conflict checking. Added ignore rules for RiderLink headers.

#rnx

diff --git a/Engine/Source/Programs/IncludeTool/IncludeTool/Program.cs b/Engine/Source/Programs/IncludeTool/IncludeTool/Program.cs
--- a/Engine/Source/Programs/IncludeTool/IncludeTool/Program.cs
+++ b/Engine/Source/Programs/IncludeTool/IncludeTool/Program.cs
@@ -1080,6 +1080,7 @@ namespace IncludeTool
 							&& !File.Location.FullName.Contains("NNI")
 							&& !File.Location.FullName.Contains("NNE")
 							&& !File.Location.FullName.Contains("PhyaLib")
+							&& !File.Location.FullName.Contains("RiderLink")
 							)
 						{
 							Log.WriteLine("{0}({1}): warning: external file is including non-external file ('{2}')", IncludedFile.Location.FullName, Markup.Location.LineIdx + 1, Markup.IncludedFile.Location);
diff --git a/Engine/Source/Programs/IncludeTool/IncludeTool/Rules.cs b/Engine/Source/Programs/IncludeTool/IncludeTool/Rules.cs
--- a/Engine/Source/Programs/IncludeTool/IncludeTool/Rules.cs
+++ b/Engine/Source/Programs/IncludeTool/IncludeTool/Rules.cs
@@ -39,6 +39,7 @@ namespace IncludeTool
 			"/engine/source/thirdparty/mcpp/mcpp-2.7.2/test-c/",
 			"/engine/source/programs/unrealswarm/private/",
 			"/engine/plugins/runtime/packethandlers/compressioncomponents/oodle/source/thirdparty/notforlicensees/oodle/213/win/examples/",
+			"/engine/plugins/developer/riderlink/",
 		};
 
 		/// <summary>
@@ -221,6 +222,11 @@ namespace IncludeTool
 			{
 				return true;
 			}
+			// Ignore Rider's RiderLink code.
+			if (NormalizedPath.Contains("/riderlink/source/riderlink/") || NormalizedPath.Contains("/riderlink/source/rd/"))
+			{
+				return true;
+			}
 			return false;
 		}
 
@@ -251,6 +257,7 @@ namespace IncludeTool
 			{ "/portal/source/layers/", "/engine/source/editor/unrealed/" },
 			{ "/portal/source/layers/", "/engine/source/runtime/online/icmp/" },
 			{ "/engine/source/programs/unreallightmass/", "/engine/source/runtime/engine/" },
+
 		};
 
 		/// <summary>
@@ -674,5 +681,24 @@ namespace IncludeTool
             }
             return true;
         }
+
+
+		static readonly string[] PathsToIgnoreConflictingSymbols = new string[] {
+			"/Engine/Source/Runtime/Core/Public/Async/TaskGraphInterfaces.h",
+			"/Engine/Source/Runtime/Core/Public/Async/TaskGraphFwd.h",
+			"/Engine/Source/Runtime/Core/Public/Delegates/DelegateAccessHandler.h",
+			"/Engine/Plugins/Runtime/RigVM/Source/RigVM/Public/RigVMCore/RigVMTypeIndex.h",
+		};
+
+		/// <summary>
+		/// Returns true if the file given should be included in report about conflicting symbols.
+		/// </summary>
+		/// <param name="Location">The file reference to check for inclusion</param>
+		/// <returns>True to report conclicts for this file</returns>
+		public static bool ReportConflictingSymbolsForFile(FileReference Location)
+		{
+			string LocationString = Location.ToString().Replace("\\", "/");
+			return !PathsToIgnoreConflictingSymbols.Any(x => LocationString.Contains(x));
+		}
 	}
 }
diff --git a/Engine/Source/Programs/IncludeTool/IncludeTool/SymbolTable.cs b/Engine/Source/Programs/IncludeTool/IncludeTool/SymbolTable.cs
--- a/Engine/Source/Programs/IncludeTool/IncludeTool/SymbolTable.cs
+++ b/Engine/Source/Programs/IncludeTool/IncludeTool/SymbolTable.cs
@@ -181,24 +181,14 @@ namespace IncludeTool
 						// UE-150382 workaround
 						// It appears that symbols are being added to Lookup whether or not branches have been taken
 						string FileName = Symbols.First().Fragment.File.Location.GetFileName();
-						if (FileName.Equals("TaskGraphInterfaces.h", StringComparison.OrdinalIgnoreCase))
+						if (Rules.ReportConflictingSymbolsForFile(Symbols.First().Fragment.File.Location))
 						{
-							continue;
-						}
-						if (FileName.Equals("RigVMTypeIndex.h", StringComparison.OrdinalIgnoreCase))
-						{
-							continue;
-						}
-						if (FileName.Equals("TaskGraphFwd.h", StringComparison.OrdinalIgnoreCase))
-						{
-							continue;
-						}
-
-						Log.WriteLine("warning: conflicting declarations of '{0}':", SymbolName);
-						foreach(Symbol Symbol in Symbols)
-						{
-							string RelativePath = Symbol.Fragment.File.Location.MakeRelativeTo(InputDir);
-							Log.WriteLine($"  {Symbol.Type} at {RelativePath}:{Symbol.Location.LineIdx+1}");
+							Log.WriteLine("warning: conflicting declarations of '{0}':", SymbolName);
+							foreach (Symbol Symbol in Symbols)
+							{
+								string RelativePath = Symbol.Fragment.File.Location.MakeRelativeTo(InputDir);
+								Log.WriteLine($"  {Symbol.Type} at {RelativePath}:{Symbol.Location.LineIdx + 1}");
+							}
 						}
 					}
 				}
