Removed references to deleted plugin "NeuralNetworkInference".

diff --git a/Engine/Source/Programs/IncludeTool/IncludeTool/Program.cs b/Engine/Source/Programs/IncludeTool/IncludeTool/Program.cs
--- a/Engine/Source/Programs/IncludeTool/IncludeTool/Program.cs
+++ b/Engine/Source/Programs/IncludeTool/IncludeTool/Program.cs
@@ -1076,8 +1076,6 @@ namespace IncludeTool
 							&& !File.Location.FullName.Contains("libSampleRate")
 							&& !File.Location.FullName.Contains("libav")
 							&& !File.Location.FullName.Contains("lz4")
-							&& !File.Location.FullName.Contains("NeuralNetworkInference")
-							&& !File.Location.FullName.Contains("NNI")
 							&& !File.Location.FullName.Contains("NNE")
 							&& !File.Location.FullName.Contains("PhyaLib")
 							&& !File.Location.FullName.Contains("RiderLink")
