RayTracing: Add "RayTracingBuiltInResources.h" to PathsToIgnoreForOldStyleHeaders in the IncludeTool until we can verify that it's safe to remove the guard (it's included in both CPP and HLSL).

#rnx

diff --git a/Engine/Source/Programs/IncludeTool/IncludeTool/Rules.cs b/Engine/Source/Programs/IncludeTool/IncludeTool/Rules.cs
--- a/Engine/Source/Programs/IncludeTool/IncludeTool/Rules.cs
+++ b/Engine/Source/Programs/IncludeTool/IncludeTool/Rules.cs
@@ -656,6 +656,7 @@ namespace IncludeTool
 			"/Engine/Plugins/Runtime/GeometryProcessing/Source/GeometryAlgorithms/Private/ThirdParty/xatlas/xatlas.h",
 
 			"/Engine/Restricted/",
+			"/Engine/Shaders/Shared/RayTracingBuiltInResources.h",
 		};
 
 		/// <summary>
