Update IncludeTool to ignore namespace in StateGraphFwd.h

diff --git a/Engine/Source/Programs/IncludeTool/IncludeTool/Rules.cs b/Engine/Source/Programs/IncludeTool/IncludeTool/Rules.cs
--- a/Engine/Source/Programs/IncludeTool/IncludeTool/Rules.cs
+++ b/Engine/Source/Programs/IncludeTool/IncludeTool/Rules.cs
@@ -447,6 +447,7 @@ namespace IncludeTool
 			"/Engine/Source/Runtime/Experimental/Chaos/Public/PhysicsProxy/SingleParticlePhysicsProxyFwd.h", // invalid forward declaration - 'namespace Chaos'
 			"/Engine/Source/Runtime/Interchange/Engine/Public/InterchangeEngineFwd.h", // invalid forward declaration - 'namespace UE'
 			"/Engine/Plugins/Experimental/GameFeatures/Source/GameFeatures/Public/GameFeatureTypesFwd.h", //  invalid forward declaration - 'namespace GameFeaturePluginStatePrivate'
+			"/Engine/Plugins/Experimental/StateGraph/Source/StateGraph/Public/StateGraphFwd.h", //  invalid forward declaration - 'namespace UE'
 			"/Engine/Source/Runtime/Core/Public/Containers/VersePathFwd.h", // invalid forward declaration - 'namespace UE::Core'
 			"/Engine/Source/Runtime/Experimental/Iris/Core/Public/Iris/ReplicationState/ReplicationStateFwd.h", // invalid forward declaration - 'namespace UE::Net'
 			"/Engine/Source/Runtime/Experimental/Iris/Core/Private/Iris/ReplicationSystem/ObjectReferenceCacheFwd.h", // invalid forward declaration - 'namespace UE::Net'
