IncludeTool: Fix parsing error when reading #error directives.

diff --git a/Engine/Source/Programs/IncludeTool/IncludeTool/SymbolTable.cs b/Engine/Source/Programs/IncludeTool/IncludeTool/SymbolTable.cs
--- a/Engine/Source/Programs/IncludeTool/IncludeTool/SymbolTable.cs
+++ b/Engine/Source/Programs/IncludeTool/IncludeTool/SymbolTable.cs
@@ -183,10 +183,10 @@ namespace IncludeTool
 						string FileName = Symbols.First().Fragment.File.Location.GetFileName();
 						if (Rules.ReportConflictingSymbolsForFile(Symbols.First().Fragment.File.Location))
 						{
-							Log.WriteLine("warning: conflicting declarations of '{0}':", SymbolName);
+						Log.WriteLine("warning: conflicting declarations of '{0}':", SymbolName);
 							foreach (Symbol Symbol in Symbols)
-							{
-								string RelativePath = Symbol.Fragment.File.Location.MakeRelativeTo(InputDir);
+						{
+							string RelativePath = Symbol.Fragment.File.Location.MakeRelativeTo(InputDir);
 								Log.WriteLine($"  {Symbol.Type} at {RelativePath}:{Symbol.Location.LineIdx + 1}");
 							}
 						}
@@ -989,9 +989,24 @@ namespace IncludeTool
 			Dictionary<Symbol, SymbolReferenceType> References = new Dictionary<Symbol, SymbolReferenceType>();
 
 			TextBuffer Text = Fragment.File.Text;
-			if(Text != null && Fragment.MarkupMax > Fragment.MarkupMin)
+			if (Text != null && Fragment.MarkupMax > Fragment.MarkupMin)
+			{
+				for (int MarkupIdx = Fragment.MarkupMin; MarkupIdx < Fragment.MarkupMax; MarkupIdx++)
+				{
+					PreprocessorMarkup Markup = Fragment.File.Markup[MarkupIdx];
+					if (Markup.Type == PreprocessorMarkupType.Text)
+					{
+						TokenReader Reader = new TokenReader(Fragment.File.Location.FullName, Fragment.File.Text, Markup.Location, Markup.EndLocation);
+						FindReferences(Reader, Fragment, References);
+					}
+				}
+			}
+
+			return References;
+		}
+
+		public void FindReferences(TokenReader Reader, SourceFragment Fragment, Dictionary<Symbol, SymbolReferenceType> References)
 			{
-				TokenReader Reader = new TokenReader(Fragment.File.Location.FullName, Fragment.File.Text, Fragment.MinLocation, Fragment.MaxLocation);
 				for(bool bMoveNext = Reader.MoveNext(TokenReaderContext.IgnoreNewlines); bMoveNext; )
 				{
 					// Read the current token, and immediately move to the next so that we can lookahead if we need to
@@ -1090,9 +1105,6 @@ namespace IncludeTool
 				}
 			}
 
-			return References;
-		}
-
 		/// <summary>
 		/// Determines whether a symbol reference should be ignored, because it's to an item defined above in the same file
 		/// </summary>
