Excluding `syms` dir from IncludeTools inspection.

diff --git a/Engine/Source/Programs/IncludeTool/IncludeTool/Rules.cs b/Engine/Source/Programs/IncludeTool/IncludeTool/Rules.cs
--- a/Engine/Source/Programs/IncludeTool/IncludeTool/Rules.cs
+++ b/Engine/Source/Programs/IncludeTool/IncludeTool/Rules.cs
@@ -39,6 +39,7 @@ namespace IncludeTool
 			"/engine/source/programs/unrealswarm/private/",
 			"/engine/plugins/runtime/packethandlers/compressioncomponents/oodle/source/thirdparty/notforlicensees/oodle/213/win/examples/",
 			"/engine/plugins/developer/riderlink/",
+			"/engine/source/runtime/symslib/syms/",
 		};
 
 		/// <summary>
@@ -226,6 +227,12 @@ namespace IncludeTool
 			{
 				return true;
 			}
+			// remove symslib files
+			if (NormalizedPath.Contains("symslib/syms"))
+			{
+				return true;
+			}
+
 			return false;
 		}
 
@@ -658,6 +665,7 @@ namespace IncludeTool
 
 			"/Engine/Restricted/",
 			"/Engine/Shaders/Shared/RayTracingBuiltInResources.h",
+			"/Engine/Source/Runtime/Symslib/syms/",
 		};
 
 		/// <summary>
